Fix the used of UMediaOutput::RequestCaptureSourceSize.
Increase the maximun number of texture buffers on the GPU from 4 to 8.
Fix the initialization of with ShouldCaptureRHITexture.
Add the capatibility to customized how the capture will occured.

diff --git a/Engine/Plugins/Media/MediaIOFramework/Source/MediaIOCore/Private/FileMediaCapture.cpp b/Engine/Plugins/Media/MediaIOFramework/Source/MediaIOCore/Private/FileMediaCapture.cpp
--- a/Engine/Plugins/Media/MediaIOFramework/Source/MediaIOCore/Private/FileMediaCapture.cpp
+++ b/Engine/Plugins/Media/MediaIOFramework/Source/MediaIOCore/Private/FileMediaCapture.cpp
@@ -9,7 +9,7 @@
 #include "Modules/ModuleManager.h"
 
 
-void UFileMediaCapture::OnFrameCaptured_RenderingThread(const FCaptureBaseData& InBaseData, TSharedPtr<FMediaCaptureUserData> InUserData, void* InBuffer, int32 Width, int32 Height)
+void UFileMediaCapture::OnFrameCaptured_RenderingThread(const FCaptureBaseData& InBaseData, TSharedPtr<FMediaCaptureUserData, ESPMode::ThreadSafe> InUserData, void* InBuffer, int32 Width, int32 Height)
 {
 	IImageWriteQueueModule* ImageWriteQueueModule = FModuleManager::Get().GetModulePtr<IImageWriteQueueModule>("ImageWriteQueue");
 	if (ImageWriteQueueModule == nullptr)
diff --git a/Engine/Plugins/Media/MediaIOFramework/Source/MediaIOCore/Private/MediaCapture.cpp b/Engine/Plugins/Media/MediaIOFramework/Source/MediaIOCore/Private/MediaCapture.cpp
--- a/Engine/Plugins/Media/MediaIOFramework/Source/MediaIOCore/Private/MediaCapture.cpp
+++ b/Engine/Plugins/Media/MediaIOFramework/Source/MediaIOCore/Private/MediaCapture.cpp
@@ -248,6 +248,7 @@ void UMediaCapture::CacheMediaOutput(EMediaCaptureSourceType InSourceType)
 {
 	check(MediaOutput);
 	DesiredSize = MediaOutput->GetRequestedSize();
+	bUseRequestedTargetSize = DesiredSize == UMediaOutput::RequestCaptureSourceSize;
 	DesiredPixelFormat = MediaOutput->GetRequestedPixelFormat();
 	ConversionOperation = MediaOutput->GetConversionOperation(InSourceType);
 }
@@ -269,6 +270,8 @@ FIntPoint UMediaCapture::GetOutputSize(const FIntPoint & InSize, const EMediaCap
 	case EMediaCaptureConversionOperation::RGB10_TO_YUVv210_10BIT:
 		// Padding aligned on 48 (16 and 6 at the same time)
 		return FIntPoint((((InSize.X + 47) / 48) * 48) / 6, InSize.Y);
+	case EMediaCaptureConversionOperation::CUSTOM:
+		return GetCustomOutputSize(InSize);
 	case EMediaCaptureConversionOperation::NONE:
 	default:
 		return InSize;
@@ -283,6 +286,8 @@ EPixelFormat UMediaCapture::GetOutputPixelFormat(const EPixelFormat & InPixelFor
 		return EPixelFormat::PF_B8G8R8A8;
 	case EMediaCaptureConversionOperation::RGB10_TO_YUVv210_10BIT:
 		return EPixelFormat::PF_R32G32B32A32_UINT;
+	case EMediaCaptureConversionOperation::CUSTOM:
+		return GetCustomOutputPixelFormat(InPixelFormat);
 	case EMediaCaptureConversionOperation::NONE:
 	default:
 		return InPixelFormat;
@@ -529,6 +534,10 @@ void UMediaCapture::InitializeResolveTarget(int32 InNumberOfBuffers)
 				This->bResolvedTargetInitialized = true;
 			});
 	}
+	else
+	{
+		bResolvedTargetInitialized = true;
+	}
 }
 
 bool UMediaCapture::ValidateMediaOutput() const
@@ -571,12 +580,17 @@ void UMediaCapture::OnEndFrame_GameThread()
 		return;
 	}
 
-	int32 ReadyFrameIndex = (CurrentResolvedTargetIndex) % NumberOfCaptureFrame; // Next one in the buffer queue
 	CurrentResolvedTargetIndex = (CurrentResolvedTargetIndex + 1) % NumberOfCaptureFrame;
+	int32 ReadyFrameIndex = (CurrentResolvedTargetIndex + 1) % NumberOfCaptureFrame; // Next one in the buffer queue
 
+	// Frame that should be on the system ram and we want to send to the user
 	FCaptureFrame* ReadyFrame = (CaptureFrames[ReadyFrameIndex].bResolvedTargetRequested) ? &CaptureFrames[ReadyFrameIndex] : nullptr;
+	// Frame that we want to transfer to system ram
 	FCaptureFrame* CapturingFrame = (GetState() != EMediaCaptureState::StopRequested) ? &CaptureFrames[CurrentResolvedTargetIndex] : nullptr;
 
+	UE_LOG(LogMediaIOCore, VeryVerbose, TEXT("MediaOutput: '%s'. ReadyFrameIndex: '%d' '%s'. CurrentResolvedTargetIndex: '%d'.")
+		, *MediaOutputName, ReadyFrameIndex, (CaptureFrames[ReadyFrameIndex].bResolvedTargetRequested) ? TEXT("Y"): TEXT("N"), CurrentResolvedTargetIndex);
+
 	if (ReadyFrame == nullptr && GetState() == EMediaCaptureState::StopRequested)
 	{
 		// All the requested frames have been captured.
@@ -766,6 +780,11 @@ void UMediaCapture::Capture_RenderThread(FRHICommandListImmediate& RHICmdList,
 				// Asynchronously copy target from GPU to GPU
 				RHICmdList.CopyToResolveTarget(SourceTexture, DestRenderTarget.TargetableTexture, ResolveParams);
 			}
+			else if (InMediaCapture->ConversionOperation == EMediaCaptureConversionOperation::CUSTOM)
+			{
+				InMediaCapture->OnCustomCapture_RenderingThread(RHICmdList, CapturingFrame->CaptureBaseData, CapturingFrame->UserData
+					, SourceTexture, DestRenderTarget.TargetableTexture, ResolveParams, {ULeft, URight}, {VTop, VBottom});
+			}
 			else
 			{
 				// convert the source with a draw call
@@ -851,16 +870,6 @@ void UMediaCapture::Capture_RenderThread(FRHICommandListImmediate& RHICmdList,
 				CapturingFrame->bResolvedTargetRequested = true;
 			}
 		}
-
-		if (InMediaCapture->bShouldCaptureRHITexture)
-		{
-			SCOPE_CYCLE_COUNTER(STAT_MediaCapture_RenderThread_Callback);
-			InMediaCapture->OnRHITextureCaptured_RenderingThread(CapturingFrame->CaptureBaseData, CapturingFrame->UserData, DestRenderTarget.TargetableTexture);
-			CapturingFrame->bResolvedTargetRequested = false;
-		}
-		else
-		{
-		}
 	}
 
 	if (!InMediaCapture->bShouldCaptureRHITexture && ReadyFrame && InMediaCapture->GetState() != EMediaCaptureState::Error)
diff --git a/Engine/Plugins/Media/MediaIOFramework/Source/MediaIOCore/Private/MediaOutput.cpp b/Engine/Plugins/Media/MediaIOFramework/Source/MediaIOCore/Private/MediaOutput.cpp
--- a/Engine/Plugins/Media/MediaIOFramework/Source/MediaIOCore/Private/MediaOutput.cpp
+++ b/Engine/Plugins/Media/MediaIOFramework/Source/MediaIOCore/Private/MediaOutput.cpp
@@ -36,13 +36,17 @@ UMediaCapture* UMediaOutput::CreateMediaCapture()
 bool UMediaOutput::Validate(FString& OutFailureReason) const
 {
 	FIntPoint RequestedSize = GetRequestedSize();
-	if (RequestedSize.X < 1 || RequestedSize.Y < 1)
+	if (RequestedSize != UMediaOutput::RequestCaptureSourceSize)
 	{
-		OutFailureReason = TEXT("The requested size is invalid.");
-		return false;
+		if (RequestedSize.X < 1 || RequestedSize.Y < 1)
+		{
+			OutFailureReason = TEXT("The requested size is invalid.");
+			return false;
+		}
 	}
 
-	if (NumberOfTextureBuffers < 1 && NumberOfTextureBuffers > 4)
+	const int32 MaxSupportedNumberOfbuffers = 8; // Arbitrary number
+	if (NumberOfTextureBuffers < 1 && NumberOfTextureBuffers > MaxSupportedNumberOfbuffers)
 	{
 		OutFailureReason = TEXT("NumberOfTextureBuffers is not valid.");
 		return false;
diff --git a/Engine/Plugins/Media/MediaIOFramework/Source/MediaIOCore/Public/FileMediaCapture.h b/Engine/Plugins/Media/MediaIOFramework/Source/MediaIOCore/Public/FileMediaCapture.h
--- a/Engine/Plugins/Media/MediaIOFramework/Source/MediaIOCore/Public/FileMediaCapture.h
+++ b/Engine/Plugins/Media/MediaIOFramework/Source/MediaIOCore/Public/FileMediaCapture.h
@@ -16,7 +16,7 @@ class MEDIAIOCORE_API UFileMediaCapture : public UMediaCapture
 	GENERATED_BODY()
 
 protected:
-	virtual void OnFrameCaptured_RenderingThread(const FCaptureBaseData& InBaseData, TSharedPtr<FMediaCaptureUserData> InUserData, void* InBuffer, int32 Width, int32 Height) override;
+	virtual void OnFrameCaptured_RenderingThread(const FCaptureBaseData& InBaseData, TSharedPtr<FMediaCaptureUserData, ESPMode::ThreadSafe> InUserData, void* InBuffer, int32 Width, int32 Height) override;
 	virtual bool CaptureSceneViewportImpl(TSharedPtr<FSceneViewport>& InSceneViewport) override;
 	virtual bool CaptureRenderTargetImpl(UTextureRenderTarget2D* InRenderTarget) override;
 
diff --git a/Engine/Plugins/Media/MediaIOFramework/Source/MediaIOCore/Public/MediaCapture.h b/Engine/Plugins/Media/MediaIOFramework/Source/MediaIOCore/Public/MediaCapture.h
--- a/Engine/Plugins/Media/MediaIOFramework/Source/MediaIOCore/Public/MediaCapture.h
+++ b/Engine/Plugins/Media/MediaIOFramework/Source/MediaIOCore/Public/MediaCapture.h
@@ -19,6 +19,7 @@
 class FSceneViewport;
 class FTextureRenderTargetResource;
 class UTextureRenderTarget2D;
+class FRHICommandListImmediate;
 
 /**
  * Possible states of media capture.
@@ -227,7 +228,7 @@ protected:
 	 * Capture the data that will pass along to the callback.
 	 * @note The capture is done on the Render Thread but triggered from the Game Thread.
 	 */
-	virtual TSharedPtr<FMediaCaptureUserData> GetCaptureFrameUserData_GameThread() { return TSharedPtr<FMediaCaptureUserData>(); }
+	virtual TSharedPtr<FMediaCaptureUserData, ESPMode::ThreadSafe> GetCaptureFrameUserData_GameThread() { return TSharedPtr<FMediaCaptureUserData, ESPMode::ThreadSafe>(); }
 
 	/** Should we call OnFrameCaptured_RenderingThread() with a RHI Texture -or- copy the memory to CPU ram and call OnFrameCaptured_RenderingThread(). */
 	virtual bool ShouldCaptureRHITexture() const { return false; }
@@ -237,14 +238,28 @@ protected:
 	 * The callback in called from a critical point. If you intend to process the buffer, do so in another thread.
 	 * The buffer is only valid for the duration of the callback.
 	 */
-	virtual void OnFrameCaptured_RenderingThread(const FCaptureBaseData& InBaseData, TSharedPtr<FMediaCaptureUserData> InUserData, void* InBuffer, int32 Width, int32 Height) { }
+	virtual void OnFrameCaptured_RenderingThread(const FCaptureBaseData& InBaseData, TSharedPtr<FMediaCaptureUserData, ESPMode::ThreadSafe> InUserData, void* InBuffer, int32 Width, int32 Height) { }
 
 	/**
 	 * Callback when the buffer was successfully copied on the GPU ram.
 	 * The callback in called from a critical point. If you intend to process the texture, do so in another thread.
 	 * The texture is valid for the duration of the callback.
 	 */
-	virtual void OnRHITextureCaptured_RenderingThread(const FCaptureBaseData& InBaseData, TSharedPtr<FMediaCaptureUserData> InUserData, FTextureRHIRef InTexture) { }
+	virtual void OnRHITextureCaptured_RenderingThread(const FCaptureBaseData& InBaseData, TSharedPtr<FMediaCaptureUserData, ESPMode::ThreadSafe> InUserData, FTextureRHIRef InTexture) { }
+
+	/**
+	 * Callback when the buffer is ready to be copied on the GPU. You need to copy yourself. You may use the callback to do texel operation or other conversion.
+	 * The callback in called from a critical point be fast.
+	 * The texture is valid for the duration of the callback.
+	 * Only called when the conversion operation is custom.
+	 */
+	virtual void OnCustomCapture_RenderingThread(FRHICommandListImmediate& RHICmdList, const FCaptureBaseData& InBaseData, TSharedPtr<FMediaCaptureUserData, ESPMode::ThreadSafe> InUserData, FTexture2DRHIRef InSourceTexture, FTextureRHIRef TargetableTexture, FResolveParams& ResolveParams, FVector2D CropU, FVector2D CropV) { }
+
+	/** Get the size of the buffer when the conversion operation is custom. The function is called once at initialization and the result is cached. */
+	virtual FIntPoint GetCustomOutputSize(const FIntPoint& InSize) const { return InSize; }
+
+	/** Get the pixel format of the buffer when the conversion operation is custom. The function is called once at initialization and the result is cached. */
+	virtual EPixelFormat GetCustomOutputPixelFormat(const EPixelFormat& InPixelFormat) const { return InPixelFormat; }
 
 protected:
 	UTextureRenderTarget2D* GetTextureRenderTarget() { return CapturingRenderTarget; }
@@ -263,8 +278,8 @@ private:
 	void OnEndFrame_GameThread();
 	void CacheMediaOutput(EMediaCaptureSourceType InSourceType);
 	void CacheOutputOptions();
-	FIntPoint GetOutputSize(const FIntPoint & InSize, const EMediaCaptureConversionOperation & InConversionOperation) const;
-	EPixelFormat GetOutputPixelFormat(const EPixelFormat & InPixelFormat, const EMediaCaptureConversionOperation & InConversionOperation) const;
+	FIntPoint GetOutputSize(const FIntPoint& InSize, const EMediaCaptureConversionOperation & InConversionOperation) const;
+	EPixelFormat GetOutputPixelFormat(const EPixelFormat& InPixelFormat, const EMediaCaptureConversionOperation & InConversionOperation) const;
 	void BroadcastStateChanged();
 	void SetFixedViewportSize(TSharedPtr<FSceneViewport> InSceneViewport);
 	void ResetFixedViewportSize(TSharedPtr<FSceneViewport> InViewport, bool bInFlushRenderingCommands);
@@ -280,8 +295,8 @@ private:
 
 		FTexture2DRHIRef ReadbackTexture;
 		FCaptureBaseData CaptureBaseData;
-		bool bResolvedTargetRequested;
-		TSharedPtr<FMediaCaptureUserData> UserData;
+		TAtomic<bool> bResolvedTargetRequested;
+		TSharedPtr<FMediaCaptureUserData, ESPMode::ThreadSafe> UserData;
 	};
 
 	TArray<FCaptureFrame> CaptureFrames;
diff --git a/Engine/Plugins/Media/MediaIOFramework/Source/MediaIOCore/Public/MediaOutput.h b/Engine/Plugins/Media/MediaIOFramework/Source/MediaIOCore/Public/MediaOutput.h
--- a/Engine/Plugins/Media/MediaIOFramework/Source/MediaIOCore/Public/MediaOutput.h
+++ b/Engine/Plugins/Media/MediaIOFramework/Source/MediaIOCore/Public/MediaOutput.h
@@ -20,6 +20,7 @@ enum class EMediaCaptureConversionOperation : int8
 	RGB10_TO_YUVv210_10BIT,
 	INVERT_ALPHA,
 	SET_ALPHA_ONE,
+	CUSTOM, //~ If Custom used, OnCustomCapture_RenderingThread, GetCustomOutputSize and GetCustomOutputPixelFormat on the MediaCapture need to be implemented
 };
 
 enum class EMediaCaptureSourceType : int8
@@ -46,7 +47,7 @@ public:
 	 * A bigger number is most likely to increase latency.
 	 * @note Some Capture are not are executed on the GPU. If it's the case then no buffer will be needed and no buffer will be created.
 	 */
-	UPROPERTY(EditAnywhere, AdvancedDisplay, Category="Output", meta=(ClampMin=1, ClampMax=4))
+	UPROPERTY(EditAnywhere, AdvancedDisplay, Category="Output", meta=(ClampMin=1, ClampMax=8))
 	int32 NumberOfTextureBuffers;
 
 	/** Creates the specific implementation of the MediaCapture for the MediaOutput. */
