Move MediaIOCore to his own plugin MediaIOFramework
Remove AJA and Blackmagic reference to MediaFrameworkUtilities.
Remove assets from AJA plugin.

diff --git a/Engine/Plugins/Media/MediaIOFramework/MediaIOFramework.uplugin b/Engine/Plugins/Media/MediaIOFramework/MediaIOFramework.uplugin
new file mode 100644
--- /dev/null
+++ b/Engine/Plugins/Media/MediaIOFramework/MediaIOFramework.uplugin
@@ -0,0 +1,28 @@
+{
+	"FileVersion" : 3,
+	"Version" : 1,
+	"VersionName" : "1.0",
+	"FriendlyName" : "Media IO Framework",
+	"Description" : "Media Framework classes to support Professional Media IO used by the Virtual Production industry.",
+	"Category" : "Media",
+	"CreatedBy" : "Epic Games, Inc.",
+	"CreatedByURL" : "http://epicgames.com",
+	"MarketplaceURL" : "",
+	"SupportURL" : "",
+	"EnabledByDefault" : false,
+	"CanContainContent" : false,
+	"IsBetaVersion" : false,
+	"Installed" : false,
+	"Modules": [
+		{
+			"Name": "MediaIOCore",
+			"Type": "Runtime",
+			"LoadingPhase": "Default"
+		},
+		{
+			"Name": "MediaIOEditor",
+			"Type": "Editor",
+			"LoadingPhase": "PostEngineInit"
+		}
+	]
+}
diff --git a/Engine/Plugins/Media/MediaIOFramework/Source/MediaIOCore/MediaIOCore.Build.cs b/Engine/Plugins/Media/MediaIOFramework/Source/MediaIOCore/MediaIOCore.Build.cs
new file mode 100644
--- /dev/null
+++ b/Engine/Plugins/Media/MediaIOFramework/Source/MediaIOCore/MediaIOCore.Build.cs
@@ -0,0 +1,37 @@
+// Copyright 1998-2018 Epic Games, Inc. All Rights Reserved.
+
+namespace UnrealBuildTool.Rules
+{
+	public class MediaIOCore : ModuleRules
+	{
+		public MediaIOCore(ReadOnlyTargetRules Target) : base(Target)
+		{
+			PublicDependencyModuleNames.AddRange(
+				new string[]
+				{
+					"Core",
+					"CoreUObject",
+					"Engine",
+					"Media",
+					"MediaAssets",
+					"MediaUtils",
+					"MovieSceneCapture",
+					"Projects",
+					"RenderCore",
+					"RHI",
+					"Slate",
+					"SlateCore",
+					"TimeManagement"
+				});
+
+			PrivateIncludePaths.AddRange(
+				new string[] {
+				});
+
+			if (Target.bBuildEditor == true)
+			{
+				PrivateDependencyModuleNames.Add("UnrealEd");
+			}
+		}
+	}
+}
diff --git a/Engine/Plugins/Media/MediaIOFramework/Source/MediaIOCore/Private/MediaCapture.cpp b/Engine/Plugins/Media/MediaIOFramework/Source/MediaIOCore/Private/MediaCapture.cpp
new file mode 100644
--- /dev/null
+++ b/Engine/Plugins/Media/MediaIOFramework/Source/MediaIOCore/Private/MediaCapture.cpp
@@ -0,0 +1,604 @@
+// Copyright 1998-2018 Epic Games, Inc. All Rights Reserved.
+
+#include "MediaCapture.h"
+
+#include "Engine/GameEngine.h"
+#include "Engine/RendererSettings.h"
+#include "Engine/TextureRenderTarget2D.h"
+#include "EngineModule.h"
+#include "GenericPlatform/GenericPlatformAtomics.h"
+#include "MediaIOCoreModule.h"
+#include "MediaOutput.h"
+#include "Misc/App.h"
+#include "RendererInterface.h"
+#include "RenderUtils.h"
+#include "Slate/SceneViewport.h"
+#include "Misc/ScopeLock.h"
+
+#if WITH_EDITOR
+#include "Editor.h"
+#include "Editor/EditorEngine.h"
+#include "Engine/Engine.h"
+#include "Engine/EngineTypes.h"
+#endif //WITH_EDITOR
+
+/* namespace MediaCaptureDetails definition
+*****************************************************************************/
+
+namespace MediaCaptureDetails
+{
+	bool FindSceneViewportAndLevel(TSharedPtr<FSceneViewport>& OutSceneViewport);
+
+	//Validation for the source of a capture
+	bool ValidateSceneViewport(const TSharedPtr<FSceneViewport>& SceneViewport, const FIntPoint& DesiredSize, const EPixelFormat DesiredPixelFormat, const bool bCurrentlyCapturing);
+	bool ValidateTextureRenderTarget2D(const UTextureRenderTarget2D* RenderTarget, const FIntPoint& DesiredSize, const EPixelFormat DesiredPixelFormat, const bool bCurrentlyCapturing);
+
+	//Validation that there is a capture 
+	bool ValidateIsCapturing(const UMediaCapture& CaptureToBeValidated);
+}
+
+
+/* UMediaCapture::FCaptureBaseData
+*****************************************************************************/
+UMediaCapture::FCaptureBaseData::FCaptureBaseData()
+	: SourceFrameNumberRenderThread(0)
+{
+
+}
+
+/* UMediaCapture::FCaptureFrame
+*****************************************************************************/
+UMediaCapture::FCaptureFrame::FCaptureFrame()
+	: bResolvedTargetRequested(false)
+{
+
+}
+
+
+/* UMediaCapture
+*****************************************************************************/
+
+UMediaCapture::UMediaCapture(const FObjectInitializer& ObjectInitializer)
+	: Super(ObjectInitializer)
+	, MediaState(EMediaCaptureState::Stopped)
+	, CurrentResolvedTargetIndex(0)
+	, NumberOfCaptureFrame(2)
+	, DesiredSize(1280, 720)
+	, DesiredPixelFormat(EPixelFormat::PF_A2B10G10R10)
+	, bResolvedTargetInitialized(false)
+	, bWaitingForResolveCommandExecution(false)
+{
+}
+
+void UMediaCapture::BeginDestroy()
+{
+	StopCapture(false);
+
+	Super::BeginDestroy();
+}
+
+FString UMediaCapture::GetDesc()
+{
+	if (MediaOutput)
+	{
+		return FString::Printf(TEXT("%s [%s]"), *Super::GetDesc(), *MediaOutput->GetDesc());
+	}
+	return FString::Printf(TEXT("%s [none]"), *Super::GetDesc());
+}
+
+bool UMediaCapture::CaptureActiveSceneViewport()
+{
+	StopCapture(false);
+
+	check(IsInGameThread());
+
+	TSharedPtr<FSceneViewport> FoundSceneViewport;
+	if (!MediaCaptureDetails::FindSceneViewportAndLevel(FoundSceneViewport) || !FoundSceneViewport.IsValid())
+	{
+		UE_LOG(LogMediaIOCore, Warning, TEXT("Can not start the capture. No viewport could be found. Play in 'Standalone' or in 'New Editor Window PIE'."));
+		return false;
+	}
+
+	return CaptureSceneViewport(FoundSceneViewport);
+}
+
+bool UMediaCapture::CaptureSceneViewport(TSharedPtr<FSceneViewport>& InSceneViewport)
+{
+	StopCapture(false);
+
+	check(IsInGameThread());
+
+	if (!ValidateMediaOutput())
+	{
+		return false;
+	}
+
+	DesiredSize = MediaOutput->GetRequestedSize();
+	DesiredPixelFormat = MediaOutput->GetRequestedPixelFormat();
+
+	const bool bCurrentlyCapturing = false;
+	if (!MediaCaptureDetails::ValidateSceneViewport(InSceneViewport, DesiredSize, DesiredPixelFormat, bCurrentlyCapturing))
+	{
+		return false;
+	}
+
+	MediaState = EMediaCaptureState::Preparing;
+	if (!CaptureSceneViewportImpl(InSceneViewport))
+	{
+		MediaState = EMediaCaptureState::Stopped;
+		return false;
+	}
+
+	//no lock required the command on the render thread is not active
+	CapturingSceneViewport = InSceneViewport;
+
+	InitializeResolveTarget(MediaOutput->NumberOfTextureBuffers);
+	CurrentResolvedTargetIndex = 0;
+	FCoreDelegates::OnEndFrame.AddUObject(this, &UMediaCapture::OnEndFrame_GameThread);
+
+	return true;
+}
+
+bool UMediaCapture::CaptureTextureRenderTarget2D(UTextureRenderTarget2D* InRenderTarget2D)
+{
+	StopCapture(false);
+
+	check(IsInGameThread());
+
+	DesiredSize = MediaOutput->GetRequestedSize();
+	DesiredPixelFormat = MediaOutput->GetRequestedPixelFormat();
+
+	const bool bCurrentlyCapturing = false;
+	if (!MediaCaptureDetails::ValidateTextureRenderTarget2D(InRenderTarget2D, DesiredSize, DesiredPixelFormat, bCurrentlyCapturing))
+	{
+		return false;
+	}
+
+	if (!CaptureRenderTargetImpl(InRenderTarget2D))
+	{
+		return false;
+	}
+
+	//no lock required the command on the render thread is not active yet
+	CapturingRenderTarget = InRenderTarget2D;
+
+	InitializeResolveTarget(MediaOutput->NumberOfTextureBuffers);
+	CurrentResolvedTargetIndex = 0;
+	FCoreDelegates::OnEndFrame.AddUObject(this, &UMediaCapture::OnEndFrame_GameThread);
+	MediaState = EMediaCaptureState::Preparing;
+
+	return true;
+}
+
+bool UMediaCapture::UpdateSceneViewport(TSharedPtr<FSceneViewport>& InSceneViewport)
+{
+	if (!MediaCaptureDetails::ValidateIsCapturing(*this))
+	{
+		StopCapture(false);
+		return false;
+	}
+
+	check(IsInGameThread());
+
+	const bool bCurrentlyCapturing = true;
+
+	if (!MediaCaptureDetails::ValidateSceneViewport(InSceneViewport, DesiredSize, DesiredPixelFormat, bCurrentlyCapturing))
+	{
+		StopCapture(false);
+		return false;
+	}
+
+	if (!UpdateSceneViewportImpl(InSceneViewport))
+	{
+		StopCapture(false);
+		return false;
+	}
+
+	{
+		FScopeLock Lock(&AccessingCapturingSource);
+		CapturingSceneViewport = InSceneViewport;
+		CapturingRenderTarget = nullptr;
+	}
+
+	return true;
+}
+
+bool UMediaCapture::UpdateTextureRenderTarget2D(UTextureRenderTarget2D * InRenderTarget2D)
+{
+	if (!MediaCaptureDetails::ValidateIsCapturing(*this))
+	{
+		StopCapture(false);
+		return false;
+	}
+
+	check(IsInGameThread());
+
+	const bool bCurrentlyCapturing = true;
+	if (!MediaCaptureDetails::ValidateTextureRenderTarget2D(InRenderTarget2D, DesiredSize, DesiredPixelFormat, bCurrentlyCapturing))
+	{
+		StopCapture(false);
+		return false;
+	}
+
+	if (!UpdateRenderTargetImpl(InRenderTarget2D))
+	{
+		StopCapture(false);
+		return false;
+	}
+
+	{
+		FScopeLock Lock(&AccessingCapturingSource);
+		CapturingRenderTarget = InRenderTarget2D;
+		CapturingSceneViewport.Reset();
+	}
+
+	return true;
+}
+
+void UMediaCapture::StopCapture(bool bAllowPendingFrameToBeProcess)
+{
+	check(IsInGameThread());
+	if (bAllowPendingFrameToBeProcess)
+	{
+		if (MediaState != EMediaCaptureState::Stopped && MediaState != EMediaCaptureState::StopRequested)
+		{
+			MediaState = EMediaCaptureState::StopRequested;
+		}
+	}
+	else
+	{
+		if (MediaState != EMediaCaptureState::Stopped)
+		{
+			MediaState = EMediaCaptureState::Stopped;
+
+			FCoreDelegates::OnEndFrame.RemoveAll(this);
+
+			while (bWaitingForResolveCommandExecution || !bResolvedTargetInitialized)
+			{
+				FlushRenderingCommands();
+			}
+			StopCaptureImpl(bAllowPendingFrameToBeProcess);
+
+			CapturingRenderTarget = nullptr;
+			CapturingSceneViewport.Reset();
+			CaptureFrames.Reset();
+			DesiredSize = FIntPoint(1280, 720);
+			DesiredPixelFormat = EPixelFormat::PF_A2B10G10R10;
+		}
+	}
+}
+
+void UMediaCapture::SetMediaOutput(UMediaOutput* InMediaOutput)
+{
+	if (GetState() == EMediaCaptureState::Stopped)
+	{
+		MediaOutput = InMediaOutput;
+	}
+}
+
+bool UMediaCapture::HasFinishedProcessing() const
+{
+	return bWaitingForResolveCommandExecution == false
+		|| MediaState == EMediaCaptureState::Error
+		|| MediaState == EMediaCaptureState::Stopped;
+}
+
+void UMediaCapture::InitializeResolveTarget(int32 InNumberOfBuffers)
+{
+	NumberOfCaptureFrame = InNumberOfBuffers;
+	check(CaptureFrames.Num() == 0);
+	CaptureFrames.AddDefaulted(InNumberOfBuffers);
+
+	auto RenderCommand = [this](FRHICommandListImmediate& RHICmdList)
+	{
+		FRHIResourceCreateInfo CreateInfo;
+		for (int32 Index = 0; Index < NumberOfCaptureFrame; ++Index)
+		{
+			CaptureFrames[Index].ReadbackTexture = RHICreateTexture2D(
+				DesiredSize.X,
+				DesiredSize.Y,
+				DesiredPixelFormat,
+				1,
+				1,
+				TexCreate_CPUReadback,
+				CreateInfo
+			);
+		}
+		bResolvedTargetInitialized = true;
+	};
+
+	ENQUEUE_UNIQUE_RENDER_COMMAND_ONEPARAMETER(
+		MediaOutputCaptureFrameCreateTexture,
+		decltype(RenderCommand), InRenderCommand, RenderCommand,
+		{
+			InRenderCommand(RHICmdList);
+		});
+}
+
+bool UMediaCapture::ValidateMediaOutput() const
+{
+	if (MediaOutput == nullptr)
+	{
+		UE_LOG(LogMediaIOCore, Error, TEXT("Can not start the capture. The Media Output is invalid."));
+		return false;
+	}
+
+	FString FailureReason;
+	if (!MediaOutput->Validate(FailureReason))
+	{
+		UE_LOG(LogMediaIOCore, Error, TEXT("Can not start the capture. %s."), *FailureReason);
+		return false;
+	}
+
+	return true;
+}
+
+void UMediaCapture::OnEndFrame_GameThread()
+{
+	if (!bResolvedTargetInitialized)
+	{
+		FlushRenderingCommands();
+	}
+
+	if (!MediaOutput)
+	{
+		return;
+	}
+
+	if (GetState() == EMediaCaptureState::Error)
+	{
+		StopCapture(false);
+	}
+
+	if (GetState() != EMediaCaptureState::Capturing && GetState() != EMediaCaptureState::StopRequested)
+	{
+		return;
+	}
+
+	int32 ReadyFrameIndex = (CurrentResolvedTargetIndex) % NumberOfCaptureFrame; // Next one in the buffer queue
+	CurrentResolvedTargetIndex = (CurrentResolvedTargetIndex + 1) % NumberOfCaptureFrame;
+
+	FCaptureFrame* ReadyFrame = (CaptureFrames[ReadyFrameIndex].bResolvedTargetRequested) ? &CaptureFrames[ReadyFrameIndex] : nullptr;
+	FCaptureFrame* CapturingFrame = (GetState() != EMediaCaptureState::StopRequested) ? &CaptureFrames[CurrentResolvedTargetIndex] : nullptr;
+
+	if (ReadyFrame == nullptr && GetState() == EMediaCaptureState::StopRequested)
+	{
+		// All the requested frames have been captured.
+		StopCapture(false);
+		return;
+	}
+
+	if (CapturingFrame)
+	{
+		//Verify if game thread is overrunning the render thread. 
+		if (CapturingFrame->bResolvedTargetRequested)
+		{
+			FlushRenderingCommands();
+		}
+
+		CapturingFrame->CaptureBaseData.SourceFrameTimecode = FApp::GetTimecode();
+		CapturingFrame->CaptureBaseData.SourceFrameNumberRenderThread = GFrameNumber;
+		CapturingFrame->UserData = GetCaptureFrameUserData_GameThread();
+	}
+
+	bWaitingForResolveCommandExecution = true;
+
+	// RenderCommand to be executed on the RenderThread
+	auto RenderCommand = [this](FRHICommandListImmediate& RHICmdList, FCaptureFrame* InCapturingFrame, FCaptureFrame* InReadyFrame)
+	{
+		FTexture2DRHIRef SourceTexture;
+		{
+			FScopeLock Lock(&AccessingCapturingSource);
+
+			UTextureRenderTarget2D* InCapturingRenderTarget = CapturingRenderTarget;
+			TSharedPtr<FSceneViewport> InSceneViewportPtr = CapturingSceneViewport.Pin();
+
+			if (InSceneViewportPtr.IsValid())
+			{
+				SourceTexture = InSceneViewportPtr->GetRenderTargetTexture();
+				if (!SourceTexture.IsValid() && InSceneViewportPtr->GetViewportRHI())
+				{
+					SourceTexture = RHICmdList.GetViewportBackBuffer(InSceneViewportPtr->GetViewportRHI());
+				}
+			}
+			else if (InCapturingRenderTarget)
+			{
+				if (InCapturingRenderTarget->GetRenderTargetResource() != nullptr && InCapturingRenderTarget->GetRenderTargetResource()->GetTextureRenderTarget2DResource() != nullptr)
+				{
+					SourceTexture = InCapturingRenderTarget->GetRenderTargetResource()->GetTextureRenderTarget2DResource()->GetTextureRHI();
+				}
+			}
+		}
+
+		if (!SourceTexture.IsValid())
+		{
+			MediaState = EMediaCaptureState::Error;
+			UMediaOutput* InMediaOutput = nullptr;
+			FPlatformAtomics::InterlockedExchangePtr((void**)(&InMediaOutput), MediaOutput);
+			UE_LOG(LogMediaIOCore, Error, TEXT("Can't grab the Texture to capture for '%s'."), InMediaOutput ? *InMediaOutput->GetName() : TEXT("[undefined]"));
+		}
+		else if (InCapturingFrame)
+		{
+			if (InCapturingFrame->ReadbackTexture->GetSizeX() != SourceTexture->GetSizeX()
+				|| InCapturingFrame->ReadbackTexture->GetSizeY() != SourceTexture->GetSizeY())
+			{
+				MediaState = EMediaCaptureState::Error;
+				UMediaOutput* InMediaOutput = nullptr;
+				FPlatformAtomics::InterlockedExchangePtr((void**)(&InMediaOutput), MediaOutput);
+				UE_LOG(LogMediaIOCore, Error, TEXT("The capture will stop for '%s'. The Source size doesn't match with the user requested size. Requested: %d,%d  Source: %d,%d")
+					, InMediaOutput ? *InMediaOutput->GetName() : TEXT("[undefined]")
+					, InCapturingFrame->ReadbackTexture->GetSizeX(), InCapturingFrame->ReadbackTexture->GetSizeY()
+					, SourceTexture->GetSizeX(), SourceTexture->GetSizeY());
+			}
+			else if (InCapturingFrame->ReadbackTexture->GetFormat() != SourceTexture->GetFormat())
+			{
+				MediaState = EMediaCaptureState::Error;
+				UMediaOutput* InMediaOutput = nullptr;
+				FPlatformAtomics::InterlockedExchangePtr((void**)(&InMediaOutput), MediaOutput);
+				UE_LOG(LogMediaIOCore, Error, TEXT("The capture will stop for '%s'. The Source pixel format doesn't match with the user requested pixel format. Requested: %s Source: %s")
+					, InMediaOutput ? *InMediaOutput->GetName() : TEXT("[undefined]")
+					, GetPixelFormatString(InCapturingFrame->ReadbackTexture->GetFormat())
+					, GetPixelFormatString(SourceTexture->GetFormat()));
+			}
+		}
+
+		if (InCapturingFrame && MediaState != EMediaCaptureState::Error)
+		{
+			FPooledRenderTargetDesc OutputDesc = FPooledRenderTargetDesc::Create2DDesc(
+				FIntPoint(SourceTexture->GetSizeX(), SourceTexture->GetSizeY()),
+				SourceTexture->GetFormat(),
+				FClearValueBinding::None,
+				TexCreate_None,
+				TexCreate_RenderTargetable,
+				false);
+			TRefCountPtr<IPooledRenderTarget> ResampleTexturePooledRenderTarget;
+			GetRendererModule().RenderTargetPoolFindFreeElement(RHICmdList, OutputDesc, ResampleTexturePooledRenderTarget, TEXT("MediaCapture"));
+			const FSceneRenderTargetItem& DestRenderTarget = ResampleTexturePooledRenderTarget->GetRenderTargetItem();
+
+			// Asynchronously copy target from GPU to GPU
+			RHICmdList.CopyToResolveTarget(SourceTexture, DestRenderTarget.TargetableTexture, FResolveParams());
+
+			// Asynchronously copy duplicate target from GPU to System Memory
+			RHICmdList.CopyToResolveTarget(DestRenderTarget.TargetableTexture, InCapturingFrame->ReadbackTexture, FResolveParams());
+
+			InCapturingFrame->bResolvedTargetRequested = true;
+		}
+
+		if (InReadyFrame && MediaState != EMediaCaptureState::Error)
+		{
+			check(InReadyFrame->ReadbackTexture.IsValid());
+
+			// Lock & read
+			void* ColorDataBuffer = nullptr;
+			int32 Width = 0, Height = 0;
+			RHICmdList.MapStagingSurface(InReadyFrame->ReadbackTexture, ColorDataBuffer, Width, Height);
+
+			OnFrameCaptured_RenderingThread(InReadyFrame->CaptureBaseData, InReadyFrame->UserData, ColorDataBuffer, Width, Height);
+			InReadyFrame->bResolvedTargetRequested = false;
+
+			RHICmdList.UnmapStagingSurface(InReadyFrame->ReadbackTexture);
+		}
+
+		bWaitingForResolveCommandExecution = false;
+	};
+
+
+	ENQUEUE_UNIQUE_RENDER_COMMAND_THREEPARAMETER(
+		MediaOutputCaptureFrameCreateTexture,
+		FCaptureFrame*, InCapturingFrame, CapturingFrame,
+		FCaptureFrame*, InPreviousFrame, ReadyFrame,
+		decltype(RenderCommand), InRenderCommand, RenderCommand,
+		{
+			InRenderCommand(RHICmdList, InCapturingFrame, InPreviousFrame);
+		});
+}
+
+/* namespace MediaCaptureDetails implementation
+*****************************************************************************/
+namespace MediaCaptureDetails
+{
+	bool FindSceneViewportAndLevel(TSharedPtr<FSceneViewport>& OutSceneViewport)
+	{
+#if WITH_EDITOR
+		if (GIsEditor)
+		{
+			for (const FWorldContext& Context : GEngine->GetWorldContexts())
+			{
+				if (Context.WorldType == EWorldType::PIE)
+				{
+					UEditorEngine* EditorEngine = CastChecked<UEditorEngine>(GEngine);
+					FSlatePlayInEditorInfo& Info = EditorEngine->SlatePlayInEditorMap.FindChecked(Context.ContextHandle);
+					if (Info.SlatePlayInEditorWindowViewport.IsValid())
+					{
+						OutSceneViewport = Info.SlatePlayInEditorWindowViewport;
+						return true;
+					}
+				}
+			}
+			return false;
+		}
+		else
+#endif
+		{
+			UGameEngine* GameEngine = CastChecked<UGameEngine>(GEngine);
+			OutSceneViewport = GameEngine->SceneViewport;
+			return true;
+		}
+	}
+
+	bool ValidateSceneViewport(const TSharedPtr<FSceneViewport>& SceneViewport, const FIntPoint& DesiredSize, const EPixelFormat DesiredPixelFormat, const bool bCurrentlyCapturing)
+	{
+
+		if (!SceneViewport.IsValid())
+		{
+			UE_LOG(LogMediaIOCore, Error, TEXT("Can not %s the capture. The Scene Viewport is invalid.")
+				, bCurrentlyCapturing ? TEXT("continue") : TEXT("start"));
+			return false;
+		}
+
+		FIntPoint SceneViewportSize = SceneViewport->GetRenderTargetTextureSizeXY();
+		if (DesiredSize.X != SceneViewportSize.X || DesiredSize.Y != SceneViewportSize.Y)
+		{
+			UE_LOG(LogMediaIOCore, Error, TEXT("Can not %s the capture. The Render Target size doesn't match with the requested size. SceneViewport: %d,%d  MediaOutput: %d,%d")
+				, bCurrentlyCapturing ? TEXT("continue") : TEXT("start")
+				, SceneViewportSize.X, SceneViewportSize.Y
+				, DesiredSize.X, DesiredSize.Y);
+			return false;
+		}
+
+		static const auto CVarDefaultBackBufferPixelFormat = IConsoleManager::Get().FindTConsoleVariableDataInt(TEXT("r.DefaultBackBufferPixelFormat"));
+		EPixelFormat SceneTargetFormat = EDefaultBackBufferPixelFormat::Convert2PixelFormat(EDefaultBackBufferPixelFormat::FromInt(CVarDefaultBackBufferPixelFormat->GetValueOnGameThread()));
+		if (DesiredPixelFormat != SceneTargetFormat)
+		{
+			UE_LOG(LogMediaIOCore, Error, TEXT("Can not %s the capture. The Render Target pixel format doesn't match with the requested pixel format. SceneViewport: %s MediaOutput: %s")
+				, bCurrentlyCapturing ? TEXT("continue") : TEXT("start")
+				, GetPixelFormatString(SceneTargetFormat)
+				, GetPixelFormatString(DesiredPixelFormat));
+			return false;
+		}
+
+		return true;
+	}
+
+	bool ValidateTextureRenderTarget2D(const UTextureRenderTarget2D* InRenderTarget2D, const FIntPoint& DesiredSize, const EPixelFormat DesiredPixelFormat, const bool bCurrentlyCapturing)
+	{
+		if (InRenderTarget2D == nullptr)
+		{
+			UE_LOG(LogMediaIOCore, Error, TEXT("Couldn't %s the capture. The Render Target is invalid.")
+				, bCurrentlyCapturing ? TEXT("continue") : TEXT("start"));
+			return false;
+		}
+
+		if (DesiredSize.X != InRenderTarget2D->SizeX || DesiredSize.Y != InRenderTarget2D->SizeY)
+		{
+			UE_LOG(LogMediaIOCore, Error, TEXT("Can not %s the capture. The Render Target size doesn't match with the requested size. RenderTarget: %d,%d  MediaOutput: %d,%d")
+				, bCurrentlyCapturing ? TEXT("continue") : TEXT("start")
+				, InRenderTarget2D->SizeX, InRenderTarget2D->SizeY
+				, DesiredSize.X, DesiredSize.Y);
+			return false;
+		}
+
+		if (DesiredPixelFormat != InRenderTarget2D->GetFormat())
+		{
+			UE_LOG(LogMediaIOCore, Error, TEXT("Can not %s the capture. The Render Target pixel format doesn't match with the requested pixel format. RenderTarget: %s MediaOutput: %s")
+				, bCurrentlyCapturing ? TEXT("continue") : TEXT("start")
+				, GetPixelFormatString(InRenderTarget2D->GetFormat())
+				, GetPixelFormatString(DesiredPixelFormat));
+			return false;
+		}
+
+		return true;
+	}
+
+	bool ValidateIsCapturing(const UMediaCapture& CaptureToBeValidated)
+	{
+		if (CaptureToBeValidated.GetState() != EMediaCaptureState::Capturing && CaptureToBeValidated.GetState() != EMediaCaptureState::Preparing)
+		{
+			UE_LOG(LogMediaIOCore, Error, TEXT("Can not update the capture. There is no capture currently.\
+			Only use UpdateSceneViewport or UpdateTextureRenderTarget2D when the state is Capturing or Preparing"));
+			return false;
+		}
+
+		return true;
+	}
+}
\ No newline at end of file
diff --git a/Engine/Plugins/Media/MediaIOFramework/Source/MediaIOCore/Private/MediaIOCoreCommonDisplayMode.cpp b/Engine/Plugins/Media/MediaIOFramework/Source/MediaIOCore/Private/MediaIOCoreCommonDisplayMode.cpp
new file mode 100644
--- /dev/null
+++ b/Engine/Plugins/Media/MediaIOFramework/Source/MediaIOCore/Private/MediaIOCoreCommonDisplayMode.cpp
@@ -0,0 +1,167 @@
+// Copyright 1998-2018 Epic Games, Inc. All Rights Reserved.
+
+#include "MediaIOCoreCommonDisplayMode.h"
+
+#define LOCTEXT_NAMESPACE "CommonDisplayMode"
+
+namespace MediaIOCommonDisplayModeInfo
+{
+	const FMediaIOCommonDisplayModeResolutionInfo AllCommonDisplayModeResolutions[] = {
+		FMediaIOCommonDisplayModeResolutionInfo{ 720,	486,	LOCTEXT("ModeRes_NTSC",		"NTSC SD")	},
+		FMediaIOCommonDisplayModeResolutionInfo{ 720,	576,	LOCTEXT("ModeRes_PAL",		"PAL SD")	},
+		FMediaIOCommonDisplayModeResolutionInfo{ 1280,	720,	LOCTEXT("ModeRes_720",		"HD 720")	},
+		FMediaIOCommonDisplayModeResolutionInfo{ 1920,	1080,	LOCTEXT("ModeRes_1080",		"HD 1080")	},
+		FMediaIOCommonDisplayModeResolutionInfo{ 2048,	1080,	LOCTEXT("ModeRes_2DCI",		"2K DCI")	},
+		FMediaIOCommonDisplayModeResolutionInfo{ 2048,	1556,	LOCTEXT("ModeRes_2Full",	"2K Full")	},
+		FMediaIOCommonDisplayModeResolutionInfo{ 3840,	2160,	LOCTEXT("ModeRes_4UHD",		"4K UHD")	},
+		FMediaIOCommonDisplayModeResolutionInfo{ 4096,	2160,	LOCTEXT("ModeRes_4DCI",		"4K DCI")	},
+		FMediaIOCommonDisplayModeResolutionInfo{ 7680,	4320,	LOCTEXT("ModeRes_8UDH",		"8K UHD")	},
+		FMediaIOCommonDisplayModeResolutionInfo{ 8192,	4320,	LOCTEXT("ModeRes_8DCI",		"8K DCI")	},
+		FMediaIOCommonDisplayModeResolutionInfo{ 15360,	8640,	LOCTEXT("ModeRes_16DCI",	"16K DCI")	},
+		FMediaIOCommonDisplayModeResolutionInfo{ 61440,	34560,	LOCTEXT("ModeRes_64DCI",	"64K DCI")	},
+		};
+
+	const FMediaIOCommonDisplayModeInfo AllCommonDisplayModes[] = {
+		FMediaIOCommonDisplayModeInfo{ 720,		486,	FFrameRate(30000, 1001),	2, LOCTEXT("Mode_NTSC", "NTSC") },
+		FMediaIOCommonDisplayModeInfo{ 720,		486,	FFrameRate(60000, 1001),	1, LOCTEXT("Mode_NTSCp", "NTSCp") },
+		FMediaIOCommonDisplayModeInfo{ 720,		576,	FFrameRate(25, 1),			2, LOCTEXT("Mode_PAL", "PAL") },
+		FMediaIOCommonDisplayModeInfo{ 720,		576,	FFrameRate(50, 1),			1, LOCTEXT("Mode_PALp", "PALp") },
+
+		FMediaIOCommonDisplayModeInfo{ 1280,	720,	FFrameRate(50, 1),			1, LOCTEXT("Mode_720p50", "720p50") },
+		FMediaIOCommonDisplayModeInfo{ 1280,	720,	FFrameRate(60000, 1001),	1, LOCTEXT("Mode_720p5994", "720p5994") },
+		FMediaIOCommonDisplayModeInfo{ 1280,	720,	FFrameRate(60, 1),			1, LOCTEXT("Mode_720p60", "720p60") },
+
+		FMediaIOCommonDisplayModeInfo{ 1920,	1080,	FFrameRate(24000, 1001),	1, LOCTEXT("Mode_1080p2398", "1080p2398") },
+		FMediaIOCommonDisplayModeInfo{ 1920,	1080,	FFrameRate(24, 1),			1, LOCTEXT("Mode_1080p24", "1080p24") },
+		FMediaIOCommonDisplayModeInfo{ 1920,	1080,	FFrameRate(25, 1),			1, LOCTEXT("Mode_1080p25", "1080p25") },
+		FMediaIOCommonDisplayModeInfo{ 1920,	1080,	FFrameRate(30000, 1001),	1, LOCTEXT("Mode_1080p2997", "1080p2997") },
+		FMediaIOCommonDisplayModeInfo{ 1920,	1080,	FFrameRate(30, 1),			1, LOCTEXT("Mode_1080p30", "1080p30") },
+		FMediaIOCommonDisplayModeInfo{ 1920,	1080,	FFrameRate(25, 1),			2, LOCTEXT("Mode_1080i50", "1080i50") },
+		FMediaIOCommonDisplayModeInfo{ 1920,	1080,	FFrameRate(30000, 1001),	2, LOCTEXT("Mode_1080i5994", "1080i5994") },
+		FMediaIOCommonDisplayModeInfo{ 1920,	1080,	FFrameRate(30, 1),			2, LOCTEXT("Mode_1080i60", "1080i60") },
+		FMediaIOCommonDisplayModeInfo{ 1920,	1080,	FFrameRate(50, 1),			1, LOCTEXT("Mode_1080p50", "1080p50") },
+		FMediaIOCommonDisplayModeInfo{ 1920,	1080,	FFrameRate(60000, 1001),	1, LOCTEXT("Mode_1080p5994", "1080p5994") },
+		FMediaIOCommonDisplayModeInfo{ 1920,	1080,	FFrameRate(60, 1),			1, LOCTEXT("Mode_1080p60", "1080p60") },
+
+		FMediaIOCommonDisplayModeInfo{ 2048,	1080,	FFrameRate(24000, 1001),	1, LOCTEXT("Mode_2kDCI2398", "2kDCI2398") },
+		FMediaIOCommonDisplayModeInfo{ 2048,	1080,	FFrameRate(24, 1),			1, LOCTEXT("Mode_2kDCI24", "2kDCI24") },
+		FMediaIOCommonDisplayModeInfo{ 2048,	1080,	FFrameRate(25, 1),			1, LOCTEXT("Mode_2kDCI25", "2kDCI25") },
+		FMediaIOCommonDisplayModeInfo{ 2048,	1080,	FFrameRate(30000, 1001),	1, LOCTEXT("Mode_2kDCI2997", "2kDCI2997") },
+		FMediaIOCommonDisplayModeInfo{ 2048,	1080,	FFrameRate(30, 1),			1, LOCTEXT("Mode_2kDCI30", "2kDCI30") },
+		FMediaIOCommonDisplayModeInfo{ 2048,	1080,	FFrameRate(50, 1),			1, LOCTEXT("Mode_2kDCI50", "2kDCI50") },
+		FMediaIOCommonDisplayModeInfo{ 2048,	1080,	FFrameRate(60000, 1001),	1, LOCTEXT("Mode_2kDCI5994", "2kDCI5994") },
+		FMediaIOCommonDisplayModeInfo{ 2048,	1080,	FFrameRate(60, 1),			1, LOCTEXT("Mode_2kDCI60", "2kDCI60") },
+
+		FMediaIOCommonDisplayModeInfo{ 2048,	1556,	FFrameRate(24000, 1001),	1, LOCTEXT("Mode_2k2398", "2k2398") },
+		FMediaIOCommonDisplayModeInfo{ 2048,	1556,	FFrameRate(24, 1),			1, LOCTEXT("Mode_2k24", "2k24") },
+		FMediaIOCommonDisplayModeInfo{ 2048,	1556,	FFrameRate(25, 1),			1, LOCTEXT("Mode_2k25", "2k25") },
+
+		FMediaIOCommonDisplayModeInfo{ 3840,	2160,	FFrameRate(24000, 1001),	1, LOCTEXT("Mode_4kUHD2398", "4kUHD2398") },
+		FMediaIOCommonDisplayModeInfo{ 3840,	2160,	FFrameRate(24, 1),			1, LOCTEXT("Mode_4kUHD24", "4kUHD24") },
+		FMediaIOCommonDisplayModeInfo{ 3840,	2160,	FFrameRate(25, 1),			1, LOCTEXT("Mode_4kUHD25", "4kUHD25") },
+		FMediaIOCommonDisplayModeInfo{ 3840,	2160,	FFrameRate(30000, 1001),	1, LOCTEXT("Mode_4kUHD2997", "4kUHD2997") },
+		FMediaIOCommonDisplayModeInfo{ 3840,	2160,	FFrameRate(30, 1),			1, LOCTEXT("Mode_4kUHD30", "4kUHD30") },
+		FMediaIOCommonDisplayModeInfo{ 3840,	2160,	FFrameRate(50, 1),			1, LOCTEXT("Mode_4kUHD50", "4kUHD50") },
+		FMediaIOCommonDisplayModeInfo{ 3840,	2160,	FFrameRate(60000, 1001),	1, LOCTEXT("Mode_4kUHD5994", "4kUHD5994") },
+		FMediaIOCommonDisplayModeInfo{ 3840,	2160,	FFrameRate(60, 1),			1, LOCTEXT("Mode_4kUHD60", "4kUHD60") },
+
+		FMediaIOCommonDisplayModeInfo{ 4096,	2160,	FFrameRate(24000, 1001),	1, LOCTEXT("Mode_4kDCI2398", "4kDCI2398") },
+		FMediaIOCommonDisplayModeInfo{ 4096,	2160,	FFrameRate(24, 1),			1, LOCTEXT("Mode_4kDCI24", "4kDCI24") },
+		FMediaIOCommonDisplayModeInfo{ 4096,	2160,	FFrameRate(25, 1),			1, LOCTEXT("Mode_4kDCI25", "4kDCI25") },
+		FMediaIOCommonDisplayModeInfo{ 4096,	2160,	FFrameRate(30000, 1001),	1, LOCTEXT("Mode_4kDCI2997", "4kDCI2997") },
+		FMediaIOCommonDisplayModeInfo{ 4096,	2160,	FFrameRate(30, 1),			1, LOCTEXT("Mode_4kDCI30", "4kDCI30") },
+		FMediaIOCommonDisplayModeInfo{ 4096,	2160,	FFrameRate(50, 1),			1, LOCTEXT("Mode_4kDCI50", "4kDCI50") },
+		FMediaIOCommonDisplayModeInfo{ 4096,	2160,	FFrameRate(60000, 1001),	1, LOCTEXT("Mode_4kDCI5994", "4kDCI5994") },
+		FMediaIOCommonDisplayModeInfo{ 4096,	2160,	FFrameRate(60, 1),			1, LOCTEXT("Mode_4kDCI60", "4kDCI60") },
+
+		FMediaIOCommonDisplayModeInfo{ 7680,	4320,	FFrameRate(24000, 1001),	1, LOCTEXT("Mode_8kUHD2398", "8kUHD2398") },
+		FMediaIOCommonDisplayModeInfo{ 7680,	4320,	FFrameRate(24, 1),			1, LOCTEXT("Mode_8kUHD24", "8kUHD24") },
+		FMediaIOCommonDisplayModeInfo{ 7680,	4320,	FFrameRate(25, 1),			1, LOCTEXT("Mode_8kUHD25", "8kUHD25") },
+		FMediaIOCommonDisplayModeInfo{ 7680,	4320,	FFrameRate(30000, 1001),	1, LOCTEXT("Mode_8kUHD2997", "8kUHD2997") },
+		FMediaIOCommonDisplayModeInfo{ 7680,	4320,	FFrameRate(30, 1),			1, LOCTEXT("Mode_8kUHD30", "8kUHD30") },
+		FMediaIOCommonDisplayModeInfo{ 7680,	4320,	FFrameRate(50, 1),			1, LOCTEXT("Mode_8kUHD50", "8kUHD50") },
+		FMediaIOCommonDisplayModeInfo{ 7680,	4320,	FFrameRate(60000, 1001),	1, LOCTEXT("Mode_8kUHD5994", "8kUHD5994") },
+		FMediaIOCommonDisplayModeInfo{ 7680,	4320,	FFrameRate(60, 1),			1, LOCTEXT("Mode_8kUHD60", "8kUHD60") },
+
+		FMediaIOCommonDisplayModeInfo{ 8192,	4320,	FFrameRate(24000, 1001),	1, LOCTEXT("Mode_8kDCI2398", "8kDCI2398") },
+		FMediaIOCommonDisplayModeInfo{ 8192,	4320,	FFrameRate(24, 1),			1, LOCTEXT("Mode_8kDCI24", "8kDCI24") },
+		FMediaIOCommonDisplayModeInfo{ 8192,	4320,	FFrameRate(25, 1),			1, LOCTEXT("Mode_8kDCI25", "8kDCI25") },
+		FMediaIOCommonDisplayModeInfo{ 8192,	4320,	FFrameRate(30000, 1001),	1, LOCTEXT("Mode_8kDCI2997", "8kDCI2997") },
+		FMediaIOCommonDisplayModeInfo{ 8192,	4320,	FFrameRate(30, 1),			1, LOCTEXT("Mode_8kDCI30", "8kDCI30") },
+		FMediaIOCommonDisplayModeInfo{ 8192,	4320,	FFrameRate(50, 1),			1, LOCTEXT("Mode_8kDCI50", "8kDCI50") },
+		FMediaIOCommonDisplayModeInfo{ 8192,	4320,	FFrameRate(60000, 1001),	1, LOCTEXT("Mode_8kDCI5994", "8kDCI5994") },
+		FMediaIOCommonDisplayModeInfo{ 8192,	4320,	FFrameRate(60, 1),			1, LOCTEXT("Mode_8kDCI60", "8kDCI60") },
+	};
+}
+
+TArrayView<const FMediaIOCommonDisplayModeInfo> FMediaIOCommonDisplayModes::GetAllModes()
+{
+	return MediaIOCommonDisplayModeInfo::AllCommonDisplayModes;
+}
+
+TArrayView<const FMediaIOCommonDisplayModeResolutionInfo> FMediaIOCommonDisplayModes::GetAllResolutions()
+{
+	return MediaIOCommonDisplayModeInfo::AllCommonDisplayModeResolutions;
+}
+
+const FMediaIOCommonDisplayModeResolutionInfo* FMediaIOCommonDisplayModes::Find(int32 InWidth, int32 InHeight)
+{
+	const int32 NumModeResolutions = ARRAY_COUNT(MediaIOCommonDisplayModeInfo::AllCommonDisplayModeResolutions);
+	for(int32 Index = 0; Index < NumModeResolutions; ++Index)
+	{
+		const FMediaIOCommonDisplayModeResolutionInfo* Info = MediaIOCommonDisplayModeInfo::AllCommonDisplayModeResolutions + Index;
+		if (Info->Width == InWidth && Info->Height == InHeight)
+		{
+			return Info;
+		}
+	}
+	return nullptr;
+}
+
+const FMediaIOCommonDisplayModeInfo* FMediaIOCommonDisplayModes::Find(int32 InWidth, int32 InHeight, const FFrameRate& InFrameRate, int32 InFieldPerFrame)
+{
+	const int32 NumModes = ARRAY_COUNT(MediaIOCommonDisplayModeInfo::AllCommonDisplayModes);
+	for (int32 Index = 0; Index < NumModes; ++Index)
+	{
+		const FMediaIOCommonDisplayModeInfo* Info = MediaIOCommonDisplayModeInfo::AllCommonDisplayModes + Index;
+		if (Info->Width == InWidth && Info->Height == InHeight && FMath::IsNearlyEqual(Info->FrameRate.AsDecimal(), InFrameRate.AsDecimal()) && Info->FieldPerFrame == InFieldPerFrame)
+		{
+			return Info;
+		}
+	}
+	return nullptr;
+}
+
+FText FMediaIOCommonDisplayModes::GetMediaIOCommonDisplayModeResolutionInfoName(int32 InWidth, int32 InHeight)
+{
+	const FMediaIOCommonDisplayModeResolutionInfo* Resolution = Find(InWidth, InHeight);
+	if (Resolution)
+	{
+		return Resolution->Name;
+	}
+	return FText::Format(LOCTEXT("Undefined_Resolution", "{0}x{1}"), FText::AsNumber(InWidth), FText::AsNumber(InHeight));
+}
+
+FText FMediaIOCommonDisplayModes::GetMediaIOCommonDisplayModeInfoName(int32 InWidth, int32 InHeight, const FFrameRate& InFrameRate, int32 InFieldPerFrame)
+{
+	const FMediaIOCommonDisplayModeInfo* Mode = Find(InWidth, InHeight, InFrameRate, InFieldPerFrame);
+	if(Mode)
+	{
+		return Mode->Name;
+	}
+
+	FNumberFormattingOptions FpsFormattingOptions;
+	FpsFormattingOptions.SetAlwaysSign(false)
+		.SetUseGrouping(false)
+		.SetMaximumFractionalDigits(2);
+
+	const FMediaIOCommonDisplayModeResolutionInfo* Resolution = Find(InWidth, InHeight);
+	if (Resolution)
+	{
+		return FText::Format(LOCTEXT("Undefined_Mode", "{0}x{1}x{2}"), Resolution->Name, FText::AsNumber(InFrameRate.AsDecimal(), &FpsFormattingOptions), InFieldPerFrame);
+	}
+
+	return FText::Format(LOCTEXT("Undefined_ModeResolution", "{0}x{1}x{2}x{3}"), FText::AsNumber(InWidth), FText::AsNumber(InHeight), FText::AsNumber(InFrameRate.AsDecimal(), &FpsFormattingOptions), InFieldPerFrame);
+}
+
+#undef LOCTEXT_NAMESPACE
diff --git a/Engine/Plugins/Media/MediaIOFramework/Source/MediaIOCore/Private/MediaIOCoreDefinitions.cpp b/Engine/Plugins/Media/MediaIOFramework/Source/MediaIOCore/Private/MediaIOCoreDefinitions.cpp
new file mode 100644
--- /dev/null
+++ b/Engine/Plugins/Media/MediaIOFramework/Source/MediaIOCore/Private/MediaIOCoreDefinitions.cpp
@@ -0,0 +1,231 @@
+// Copyright 1998-2018 Epic Games, Inc. All Rights Reserved.
+
+#include "MediaIOCoreDefinitions.h"
+#include "Containers/ArrayView.h"
+#include "MediaIOCoreCommonDisplayMode.h"
+#include "Misc/FrameRate.h"
+
+#define LOCTEXT_NAMESPACE "MediaIOCoreDefinitions"
+
+/**
+ * MediaIOCoreDefinitions
+ */
+namespace MediaIOCoreDefinitions
+{
+	const int32 InvalidDeviceIdentifier = -1;
+	const int32 InvalidDevicePortIdentifier = -1;
+	const int32 InvalidDeviceModeIdentifier = -1;
+	
+	const FName NAME_Protocol = TEXT("protocol");
+	
+	const TCHAR* DeviceStr = TEXT("device");
+	const TCHAR* SingleStr = TEXT("single");
+	const TCHAR* DualStr = TEXT("dual");
+	const TCHAR* QuadSquareStr = TEXT("quadSQ");
+	const TCHAR* QuadTsiStr = TEXT("quadSI");
+	
+	const TCHAR* GetTransportString(EMediaIOSDITransportType InLinkType, EMediaIOQuadLinkTransportType InQuadLinkType)
+	{
+		const TCHAR* Result = MediaIOCoreDefinitions::SingleStr;
+		switch (InLinkType)
+		{
+		case EMediaIOSDITransportType::DualLink: Result = MediaIOCoreDefinitions::DualStr; break;
+		case EMediaIOSDITransportType::QuadLink:
+		{
+			Result = MediaIOCoreDefinitions::QuadSquareStr;
+			if (InQuadLinkType == EMediaIOQuadLinkTransportType::TwoSampleInterleave)
+			{
+				Result = MediaIOCoreDefinitions::QuadTsiStr;
+			}
+			break;
+		}
+		}
+		return Result;
+	}
+}
+
+
+/**
+ * FMediaIOCoreDevice
+ */
+FMediaIODevice::FMediaIODevice()
+	: DeviceIdentifier(MediaIOCoreDefinitions::InvalidDeviceIdentifier)
+{}
+
+
+bool FMediaIODevice::operator==(const FMediaIODevice& Other) const
+{
+	return Other.DeviceIdentifier == DeviceIdentifier;
+}
+
+
+bool FMediaIODevice::IsValid() const
+{
+	return DeviceIdentifier != MediaIOCoreDefinitions::InvalidDeviceIdentifier;
+}
+
+
+/**
+ * FMediaIOConnection
+ */
+FMediaIOConnection::FMediaIOConnection()
+	: Protocol(MediaIOCoreDefinitions::NAME_Protocol)
+	, TransportType(EMediaIOSDITransportType::SingleLink)
+	, QuadTransportType(EMediaIOQuadLinkTransportType::SquareDivision)
+	, PortIdentifier(MediaIOCoreDefinitions::InvalidDevicePortIdentifier)
+{}
+
+
+bool FMediaIOConnection::operator==(const FMediaIOConnection& Other) const
+{
+	return Other.Device == Device
+		&& Other.TransportType == TransportType
+		&& Other.PortIdentifier == PortIdentifier
+		&& (Other.TransportType != EMediaIOSDITransportType::QuadLink || Other.QuadTransportType == QuadTransportType);
+}
+
+
+FString FMediaIOConnection::ToUrl() const
+{
+	if (IsValid())
+	{
+		return FString::Printf(TEXT("%s://%s%d/%s%d")
+			, *Protocol.ToString()
+			, MediaIOCoreDefinitions::DeviceStr
+			, Device.DeviceIdentifier
+			, MediaIOCoreDefinitions::GetTransportString(TransportType, QuadTransportType)
+			, PortIdentifier
+			);
+	}
+	return Protocol.ToString();
+}
+
+
+bool FMediaIOConnection::IsValid() const
+{
+	return Device.IsValid() && PortIdentifier != MediaIOCoreDefinitions::InvalidDevicePortIdentifier;
+}
+
+/**
+ * FMediaIOMode
+ */
+FMediaIOMode::FMediaIOMode()
+	: DeviceModeIdentifier(MediaIOCoreDefinitions::InvalidDeviceModeIdentifier)
+{}
+
+
+bool FMediaIOMode::operator==(const FMediaIOMode& Other) const
+{
+	return Other.DeviceModeIdentifier == DeviceModeIdentifier;
+}
+
+
+FText FMediaIOMode::GetModeName() const
+{
+	if (IsValid())
+	{
+		int32 FieldPerFrame = 1;
+		FFrameRate FieldFrameRate = FrameRate;
+		if (Standard == EMediaIOStandardType::Interlaced)
+		{
+			FieldPerFrame = 2;
+			FieldFrameRate.Numerator /= 2;
+		}
+		return FMediaIOCommonDisplayModes::GetMediaIOCommonDisplayModeInfoName(Resolution.X, Resolution.Y, FieldFrameRate, FieldPerFrame);
+	}
+	return LOCTEXT("Invalid", "<Invalid>");
+}
+
+
+bool FMediaIOMode::IsValid() const
+{
+	return DeviceModeIdentifier != MediaIOCoreDefinitions::InvalidDeviceModeIdentifier;
+}
+
+
+/**
+ * FMediaIOConfiguration
+ */
+FMediaIOConfiguration::FMediaIOConfiguration()
+	: bIsInput(true)
+{}
+
+
+bool FMediaIOConfiguration::operator== (const FMediaIOConfiguration& Other) const
+{
+	return MediaConnection == Other.MediaConnection
+		&& MediaMode == Other.MediaMode
+		&& bIsInput == Other.bIsInput;
+}
+
+
+bool FMediaIOConfiguration::IsValid() const
+{
+	return MediaConnection.IsValid() && MediaMode.IsValid();
+}
+
+
+/**
+ * FMediaIOConfiguration
+ */
+FMediaIOOutputConfiguration::FMediaIOOutputConfiguration()
+	: OutputType(EMediaIOOutputType::Fill)
+	, KeyPortIdentifier(MediaIOCoreDefinitions::InvalidDevicePortIdentifier)
+	, OutputReference(EMediaIOReferenceType::FreeRun)
+	, ReferencePortIdentifier(MediaIOCoreDefinitions::InvalidDevicePortIdentifier)
+{
+	MediaConfiguration.bIsInput = false;
+}
+
+
+bool FMediaIOOutputConfiguration::operator== (const FMediaIOOutputConfiguration& Other) const
+{
+	if (OutputType == Other.OutputType
+		&& MediaConfiguration == Other.MediaConfiguration
+		&& OutputReference == Other.OutputReference)
+	{
+		if (OutputType == EMediaIOOutputType::FillAndKey)
+		{
+			if (KeyPortIdentifier != Other.KeyPortIdentifier)
+			{
+				return false;
+			}
+		}
+		if (OutputReference == EMediaIOReferenceType::Input)
+		{
+			if (ReferencePortIdentifier != Other.ReferencePortIdentifier)
+			{
+				return false;
+			}
+		}
+		return true;
+	}
+	return false;
+}
+
+
+bool FMediaIOOutputConfiguration::IsValid() const
+{
+	bool bResult = false;
+	if (MediaConfiguration.IsValid())
+	{
+		bResult = true;
+		if (OutputType == EMediaIOOutputType::FillAndKey)
+		{
+			if (KeyPortIdentifier == MediaIOCoreDefinitions::InvalidDevicePortIdentifier)
+			{
+				bResult = false;
+			}
+		}
+		if (OutputReference == EMediaIOReferenceType::Input)
+		{
+			if (ReferencePortIdentifier == MediaIOCoreDefinitions::InvalidDevicePortIdentifier)
+			{
+				bResult = false;
+			}
+		}
+	}
+	return bResult;
+}
+
+#undef LOCTEXT_NAMESPACE
diff --git a/Engine/Plugins/Media/MediaIOFramework/Source/MediaIOCore/Private/MediaIOCoreDeviceProvider.cpp b/Engine/Plugins/Media/MediaIOFramework/Source/MediaIOCore/Private/MediaIOCoreDeviceProvider.cpp
new file mode 100644
--- /dev/null
+++ b/Engine/Plugins/Media/MediaIOFramework/Source/MediaIOCore/Private/MediaIOCoreDeviceProvider.cpp
@@ -0,0 +1,123 @@
+// Copyright 1998-2018 Epic Games, Inc. All Rights Reserved.
+
+#include "IMediaIOCoreDeviceProvider.h"
+#include "MediaIOCoreCommonDisplayMode.h"
+
+
+#define LOCTEXT_NAMESPACE "MediaIOCoreDeviceProvider"
+
+
+FText IMediaIOCoreDeviceProvider::GetTransportName(EMediaIOSDITransportType InLinkType, EMediaIOQuadLinkTransportType InQuadLinkType)
+{
+	FText Transport = LOCTEXT("Single", "single");
+	if (InLinkType == EMediaIOSDITransportType::DualLink)
+	{
+		Transport = LOCTEXT("Dual", "dual");
+	}
+	else if (InLinkType == EMediaIOSDITransportType::QuadLink && InQuadLinkType == EMediaIOQuadLinkTransportType::SquareDivision)
+	{
+		Transport = LOCTEXT("QuadSquare", "quadSQ");
+	}
+	else if (InLinkType == EMediaIOSDITransportType::QuadLink && InQuadLinkType == EMediaIOQuadLinkTransportType::TwoSampleInterleave)
+	{
+		Transport = LOCTEXT("QuadTSI", "quadSI");
+	}
+	return Transport;
+}
+
+FText IMediaIOCoreDeviceProvider::ToText(const FMediaIOConfiguration& InConfiguration) const
+{
+	if (InConfiguration.IsValid())
+	{
+		return FText::Format(LOCTEXT("FMediaIOConfigurationToText", "{0} - {1} [device{2}/{3}{4}/{5}]")
+			, InConfiguration.bIsInput ? LOCTEXT("In", "In") : LOCTEXT("Out", "Out")
+			, FText::FromName(InConfiguration.MediaConnection.Device.DeviceName)
+			, FText::AsNumber(InConfiguration.MediaConnection.Device.DeviceIdentifier)
+			, GetTransportName(InConfiguration.MediaConnection.TransportType, InConfiguration.MediaConnection.QuadTransportType)
+			, FText::AsNumber(InConfiguration.MediaConnection.PortIdentifier)
+			, InConfiguration.MediaMode.GetModeName()
+			);
+	}
+	return LOCTEXT("Invalid", "<Invalid>");
+}
+
+
+FText IMediaIOCoreDeviceProvider::ToText(const FMediaIOConnection& InConnection) const
+{
+	if (InConnection.IsValid())
+	{
+		return FText::Format(LOCTEXT("FMediaIOConnectionToText", "{0} [device{1}/{2}{3}]")
+			, FText::FromName(InConnection.Device.DeviceName)
+			, FText::AsNumber(InConnection.Device.DeviceIdentifier)
+			, GetTransportName(InConnection.TransportType, InConnection.QuadTransportType)
+			, FText::AsNumber(InConnection.PortIdentifier)
+			);
+	}
+	return LOCTEXT("Invalid", "<Invalid>");
+}
+
+
+FText IMediaIOCoreDeviceProvider::ToText(const FMediaIODevice& InDevice) const
+{
+	if (InDevice.IsValid())
+	{
+		return FText::Format(LOCTEXT("FMediaIODeviceToText", "{0} [device{1}]")
+				, FText::FromName(InDevice.DeviceName)
+				, InDevice.DeviceIdentifier
+				);
+	}
+	return LOCTEXT("Invalid", "<Invalid>");
+}
+
+
+FText IMediaIOCoreDeviceProvider::ToText(const FMediaIOMode& InMode) const
+{
+	return InMode.GetModeName();
+}
+
+
+FText IMediaIOCoreDeviceProvider::ToText(const FMediaIOOutputConfiguration& InConfiguration) const
+{
+	if (InConfiguration.IsValid())
+	{
+		FText KeyText = FText::GetEmpty();
+
+		if (InConfiguration.OutputType == EMediaIOOutputType::FillAndKey)
+		{
+			KeyText = FText::Format(LOCTEXT("FMediaIOOutputConfigurationKeyToText", " - Key[{0}{1}]")
+				, GetTransportName(InConfiguration.MediaConfiguration.MediaConnection.TransportType, InConfiguration.MediaConfiguration.MediaConnection.QuadTransportType)
+				, FText::AsNumber(InConfiguration.KeyPortIdentifier)
+				);
+		}
+
+		FText ReferenceText = FText::GetEmpty();
+		if (InConfiguration.OutputReference == EMediaIOReferenceType::Input)
+		{
+			ReferenceText = FText::Format(LOCTEXT("FMediaIOOutputConfigurationRefToText", " - Ref[{0}{1}]")
+				, GetTransportName(InConfiguration.MediaConfiguration.MediaConnection.TransportType, InConfiguration.MediaConfiguration.MediaConnection.QuadTransportType)
+				, FText::AsNumber(InConfiguration.ReferencePortIdentifier)
+			);
+		}
+		else if (InConfiguration.OutputReference == EMediaIOReferenceType::External)
+		{
+			ReferenceText = LOCTEXT("External", "External");
+		}
+		else if (InConfiguration.OutputReference == EMediaIOReferenceType::FreeRun)
+		{
+			ReferenceText = LOCTEXT("FreeRun", "FreeRun");
+		}
+
+		return FText::Format(LOCTEXT("FMediaIOOutputConfigurationToText", "Fill - {0} [device{1}/{2}{3}/{4}]{5}{6}")
+			, FText::FromName(InConfiguration.MediaConfiguration.MediaConnection.Device.DeviceName)
+			, FText::AsNumber(InConfiguration.MediaConfiguration.MediaConnection.Device.DeviceIdentifier)
+			, GetTransportName(InConfiguration.MediaConfiguration.MediaConnection.TransportType, InConfiguration.MediaConfiguration.MediaConnection.QuadTransportType)
+			, FText::AsNumber(InConfiguration.MediaConfiguration.MediaConnection.PortIdentifier)
+			, InConfiguration.MediaConfiguration.MediaMode.GetModeName()
+			, KeyText
+			, ReferenceText
+			);
+	}
+	return LOCTEXT("Invalid", "<Invalid>");
+}
+
+#undef LOCTEXT_NAMESPACE
diff --git a/Engine/Plugins/Media/MediaIOFramework/Source/MediaIOCore/Private/MediaIOCoreModule.cpp b/Engine/Plugins/Media/MediaIOFramework/Source/MediaIOCore/Private/MediaIOCoreModule.cpp
new file mode 100644
--- /dev/null
+++ b/Engine/Plugins/Media/MediaIOFramework/Source/MediaIOCore/Private/MediaIOCoreModule.cpp
@@ -0,0 +1,54 @@
+// Copyright 1998-2018 Epic Games, Inc. All Rights Reserved.
+
+#include "MediaIOCoreModule.h"
+#include "IMediaIOCoreDeviceProvider.h"
+#include "IMediaIOCoreModule.h"
+
+#include "Modules/ModuleManager.h"
+
+
+DEFINE_LOG_CATEGORY(LogMediaIOCore);
+
+/**
+ * Implements the MediaIOCore module.
+ */
+class FMediaIOCoreModule : public IMediaIOCoreModule
+{
+public:
+	virtual void RegisterDeviceProvider(IMediaIOCoreDeviceProvider* InProvider) override
+	{
+		if (InProvider)
+		{
+			DeviceProviders.Add(InProvider);
+		}
+	}
+
+	virtual void UnregisterDeviceProvider(IMediaIOCoreDeviceProvider* InProvider) override
+	{
+		DeviceProviders.RemoveSingleSwap(InProvider);
+	}
+
+	virtual IMediaIOCoreDeviceProvider* GetDeviceProvider(FName InProviderName) override
+	{
+		for(IMediaIOCoreDeviceProvider* DeviceProvider : DeviceProviders)
+		{
+			if (DeviceProvider->GetFName() == InProviderName)
+			{
+				return DeviceProvider;
+			}
+		}
+		return nullptr;
+	}
+
+private:
+	TArray<IMediaIOCoreDeviceProvider*> DeviceProviders;
+};
+
+
+IMediaIOCoreModule& IMediaIOCoreModule::Get()
+{
+	return FModuleManager::LoadModuleChecked<FMediaIOCoreModule>("MediaIOCore");
+}
+
+
+IMPLEMENT_MODULE(FMediaIOCoreModule, MediaIOCore);
diff --git a/Engine/Plugins/Media/MediaIOFramework/Source/MediaIOCore/Private/MediaIOCoreModule.h b/Engine/Plugins/Media/MediaIOFramework/Source/MediaIOCore/Private/MediaIOCoreModule.h
new file mode 100644
--- /dev/null
+++ b/Engine/Plugins/Media/MediaIOFramework/Source/MediaIOCore/Private/MediaIOCoreModule.h
@@ -0,0 +1,6 @@
+// Copyright 1998-2018 Epic Games, Inc. All Rights Reserved.
+#pragma once
+
+#include "CoreMinimal.h"
+
+DECLARE_LOG_CATEGORY_EXTERN(LogMediaIOCore, Log, All);
diff --git a/Engine/Plugins/Media/MediaIOFramework/Source/MediaIOCore/Private/MediaOutput.cpp b/Engine/Plugins/Media/MediaIOFramework/Source/MediaIOCore/Private/MediaOutput.cpp
new file mode 100644
--- /dev/null
+++ b/Engine/Plugins/Media/MediaIOFramework/Source/MediaIOCore/Private/MediaOutput.cpp
@@ -0,0 +1,44 @@
+// Copyright 1998-2018 Epic Games, Inc. All Rights Reserved.
+
+#include "MediaOutput.h"
+
+#include "MediaCapture.h"
+#include "MediaIOCoreModule.h"
+
+/* IMediaOptions interface
+ *****************************************************************************/
+
+UMediaOutput::UMediaOutput(const FObjectInitializer& ObjectInitializer)
+	: Super(ObjectInitializer)
+	, NumberOfTextureBuffers(2)
+{
+}
+
+UMediaCapture* UMediaOutput::CreateMediaCapture()
+{
+	UMediaCapture* Result = nullptr;
+
+	FString FailureReason;
+	if (Validate(FailureReason))
+	{
+		Result = CreateMediaCaptureImpl();
+	}
+	else
+	{
+		UE_LOG(LogMediaIOCore, Error, TEXT("Couldn't create the media capture. %s."), *FailureReason);
+	}
+
+	return Result;
+}
+
+bool UMediaOutput::Validate(FString& OutFailureReason) const
+{
+	bool bResult = (NumberOfTextureBuffers >= 1 && NumberOfTextureBuffers <= 4);
+
+	if (!bResult)
+	{
+		OutFailureReason = TEXT("NumberOfBuffer is not valid.");
+	}
+
+	return bResult;
+}
diff --git a/Engine/Plugins/Media/MediaIOFramework/Source/MediaIOCore/Private/Player/MediaIOCorePlayerBase.cpp b/Engine/Plugins/Media/MediaIOFramework/Source/MediaIOCore/Private/Player/MediaIOCorePlayerBase.cpp
new file mode 100644
--- /dev/null
+++ b/Engine/Plugins/Media/MediaIOFramework/Source/MediaIOCore/Private/Player/MediaIOCorePlayerBase.cpp
@@ -0,0 +1,423 @@
+// Copyright 1998-2018 Epic Games, Inc. All Rights Reserved.
+
+#include "MediaIOCorePlayerBase.h"
+
+#include "Engine/Engine.h"
+#include "Engine/TimecodeProvider.h"
+#include "IMediaEventSink.h"
+#include "IMediaOptions.h"
+#include "MediaIOCoreSamples.h"
+#include "Misc/App.h"
+#include "Misc/ScopeLock.h"
+#include "TimeSynchronizableMediaSource.h"
+
+
+#define LOCTEXT_NAMESPACE "MediaIOCorePlayerBase"
+
+/* FMediaIOCoreMediaOption structors
+ *****************************************************************************/
+const FName FMediaIOCoreMediaOption::FrameRateNumerator("FrameRateNumerator");
+const FName FMediaIOCoreMediaOption::FrameRateDenominator("FrameRateDenominator");
+const FName FMediaIOCoreMediaOption::ResolutionWidth("ResolutionWidth");
+const FName FMediaIOCoreMediaOption::ResolutionHeight("ResolutionHeight");
+const FName FMediaIOCoreMediaOption::VideoModeName("VideoModeName");
+
+/* FMediaIOCorePlayerBase structors
+ *****************************************************************************/
+
+FMediaIOCorePlayerBase::FMediaIOCorePlayerBase(IMediaEventSink& InEventSink)
+	: bIsTimecodeLogEnable(false)
+	, CurrentState(EMediaState::Closed)
+	, CurrentTime(FTimespan::Zero())
+	, EventSink(InEventSink)
+	, VideoFrameRate(30, 1)
+	, LastFrameDropCount(0)
+	, Samples(new FMediaIOCoreSamples)
+	, bUseTimeSynchronization(false)
+	, PreviousFrameTimespan(FTimespan::Zero())
+{
+}
+
+FMediaIOCorePlayerBase::~FMediaIOCorePlayerBase()
+{
+	delete Samples;
+	Samples = nullptr;
+}
+
+/* IMediaPlayer interface
+*****************************************************************************/
+
+void FMediaIOCorePlayerBase::Close()
+{
+	CurrentState = EMediaState::Closed;
+	CurrentTime = FTimespan::Zero();
+	AudioTrackFormat.NumChannels = 0;
+	AudioTrackFormat.SampleRate = 0;
+
+	Samples->FlushSamples();
+	EventSink.ReceiveMediaEvent(EMediaEvent::TracksChanged);
+	EventSink.ReceiveMediaEvent(EMediaEvent::MediaClosed);
+}
+
+FString FMediaIOCorePlayerBase::GetInfo() const
+{
+	FString Info;
+
+	if (AudioTrackFormat.NumChannels > 0)
+	{
+		Info += FString::Printf(TEXT("Stream\n"));
+		Info += FString::Printf(TEXT("    Type: Audio\n"));
+		Info += FString::Printf(TEXT("    Channels: %i\n"), AudioTrackFormat.NumChannels);
+		Info += FString::Printf(TEXT("    Sample Rate: %i Hz\n"), AudioTrackFormat.SampleRate);
+		Info += FString::Printf(TEXT("    Bits Per Sample: 32\n"));
+	}
+
+	if (VideoTrackFormat.Dim != FIntPoint::ZeroValue)
+	{
+		if (!Info.IsEmpty())
+		{
+			Info += TEXT("\n");
+		}
+		Info += FString::Printf(TEXT("Stream\n"));
+		Info += FString::Printf(TEXT("    Type: Video\n"));
+		Info += FString::Printf(TEXT("    Dimensions: %i x %i\n"), VideoTrackFormat.Dim.X, VideoTrackFormat.Dim.Y);
+		Info += FString::Printf(TEXT("    Frame Rate: %g fps\n"), VideoFrameRate.AsDecimal());
+	}
+	return Info;
+}
+
+IMediaCache& FMediaIOCorePlayerBase::GetCache()
+{
+	return *this;
+}
+
+IMediaControls& FMediaIOCorePlayerBase::GetControls()
+{
+	return *this;
+}
+
+IMediaSamples& FMediaIOCorePlayerBase::GetSamples()
+{
+	return *Samples;
+}
+
+const FMediaIOCoreSamples& FMediaIOCorePlayerBase::GetSamples() const
+{
+	return *Samples;
+}
+
+FString FMediaIOCorePlayerBase::GetStats() const
+{
+	return FString();
+}
+
+IMediaTracks& FMediaIOCorePlayerBase::GetTracks()
+{
+	return *this;
+}
+
+FString FMediaIOCorePlayerBase::GetUrl() const
+{
+	return OpenUrl;
+}
+
+IMediaView& FMediaIOCorePlayerBase::GetView()
+{
+	return *this;
+}
+
+bool FMediaIOCorePlayerBase::Open(const FString& Url, const IMediaOptions* Options)
+{
+	Close();
+
+	OpenUrl = Url;
+	return ReadMediaOptions(Options);
+}
+
+bool FMediaIOCorePlayerBase::Open(const TSharedRef<FArchive, ESPMode::ThreadSafe>& /*Archive*/, const FString& /*OriginalUrl*/, const IMediaOptions* /*Options*/)
+{
+	return false;
+}
+
+void FMediaIOCorePlayerBase::TickTimeManagement()
+{
+	bool bUseTimecode = false;
+	if (bUseTimeSynchronization)
+	{
+		if (const UTimecodeProvider* Provider = GEngine->GetTimecodeProvider())
+		{
+			bUseTimecode = (Provider->GetSynchronizationState() == ETimecodeProviderSynchronizationState::Synchronized);
+		}
+	}
+
+	if (bUseTimecode)
+	{
+		FTimecode Timecode = FApp::GetTimecode();
+		FFrameRate FrameRate = FApp::GetTimecodeFrameRate();
+		CurrentTime = FTimespan(0, Timecode.Hours, Timecode.Minutes, Timecode.Seconds, static_cast<int32>((ETimespan::TicksPerSecond * Timecode.Frames) / FrameRate.AsDecimal()) * ETimespan::NanosecondsPerTick);
+	}
+	else
+	{
+		// As default, use the App time
+		CurrentTime = FTimespan::FromSeconds(FApp::GetCurrentTime());
+	}
+}
+
+/* IMediaCache interface
+ *****************************************************************************/
+bool FMediaIOCorePlayerBase::QueryCacheState(EMediaCacheState State, TRangeSet<FTimespan>& OutTimeRanges) const
+{
+	if (!Samples || Samples->NumVideoSamples() <= 0)
+	{
+		return false;
+	}
+	
+	bool bHasQueried = false;
+	if (State == EMediaCacheState::Loaded)
+	{
+		const FTimespan FrameDuration = FTimespan::FromSeconds(VideoFrameRate.AsInterval());
+		const FTimespan NextSampleTime = Samples->GetNextVideoSampleTime();
+		OutTimeRanges.Add(TRange<FTimespan>(NextSampleTime, NextSampleTime + FrameDuration * Samples->NumVideoSamples()));
+		bHasQueried = true;
+	}
+
+	return bHasQueried;
+}
+
+int32 FMediaIOCorePlayerBase::GetSampleCount(EMediaCacheState State) const
+{
+	int32 Count = 0;
+	if (State == EMediaCacheState::Loaded)
+	{
+		if (Samples)
+		{
+			Count = Samples->NumVideoSamples();
+		}
+	}
+
+	return Count;
+}
+
+
+/* IMediaControls interface
+ *****************************************************************************/
+
+bool FMediaIOCorePlayerBase::CanControl(EMediaControl Control) const
+{
+	return false;
+}
+
+
+FTimespan FMediaIOCorePlayerBase::GetDuration() const
+{
+	return (CurrentState == EMediaState::Playing) ? FTimespan::MaxValue() : FTimespan::Zero();
+}
+
+
+float FMediaIOCorePlayerBase::GetRate() const
+{
+	return (CurrentState == EMediaState::Playing) ? 1.0f : 0.0f;
+}
+
+
+EMediaState FMediaIOCorePlayerBase::GetState() const
+{
+	return CurrentState;
+}
+
+
+EMediaStatus FMediaIOCorePlayerBase::GetStatus() const
+{
+	return (CurrentState == EMediaState::Preparing) ? EMediaStatus::Connecting : EMediaStatus::None;
+}
+
+TRangeSet<float> FMediaIOCorePlayerBase::GetSupportedRates(EMediaRateThinning /*Thinning*/) const
+{
+	TRangeSet<float> Result;
+	Result.Add(TRange<float>(0.0f));
+	Result.Add(TRange<float>(1.0f));
+	return Result;
+}
+
+FTimespan FMediaIOCorePlayerBase::GetTime() const
+{
+	return CurrentTime;
+}
+
+
+bool FMediaIOCorePlayerBase::IsLooping() const
+{
+	return false; // not supported
+}
+
+
+bool FMediaIOCorePlayerBase::Seek(const FTimespan& Time)
+{
+	return false; // not supported
+}
+
+
+bool FMediaIOCorePlayerBase::SetLooping(bool Looping)
+{
+	return false; // not supported
+}
+
+
+bool FMediaIOCorePlayerBase::SetRate(float Rate)
+{
+	return false;
+}
+
+
+/* IMediaTracks interface
+ *****************************************************************************/
+
+bool FMediaIOCorePlayerBase::GetAudioTrackFormat(int32 TrackIndex, int32 FormatIndex, FMediaAudioTrackFormat& OutFormat) const
+{
+	if (!IsHardwareReady() || TrackIndex != 0 || FormatIndex != 0)
+	{
+		return false;
+	}
+
+	OutFormat = AudioTrackFormat;
+	return true;
+}
+
+
+int32 FMediaIOCorePlayerBase::GetNumTracks(EMediaTrackType TrackType) const
+{
+	return 1;
+}
+
+int32 FMediaIOCorePlayerBase::GetNumTrackFormats(EMediaTrackType TrackType, int32 TrackIndex) const
+{
+	return 1;
+}
+
+
+int32 FMediaIOCorePlayerBase::GetSelectedTrack(EMediaTrackType TrackType) const
+{
+	switch (TrackType)
+	{
+	case EMediaTrackType::Audio:
+	case EMediaTrackType::Video:
+		return 0;
+
+	default:
+		return INDEX_NONE;
+	}
+}
+
+
+FText FMediaIOCorePlayerBase::GetTrackDisplayName(EMediaTrackType TrackType, int32 TrackIndex) const
+{
+	if (!IsHardwareReady() || TrackIndex != 0)
+	{
+		return FText::GetEmpty();
+	}
+
+	switch (TrackType)
+	{
+	case EMediaTrackType::Audio:
+		return LOCTEXT("DefaultAudioTrackName", "Audio Track");
+
+	case EMediaTrackType::Video:
+		return LOCTEXT("DefaultVideoTrackName", "Video Track");
+
+	default:
+		return FText::GetEmpty();
+	}
+
+	return FText::GetEmpty();
+}
+
+
+int32 FMediaIOCorePlayerBase::GetTrackFormat(EMediaTrackType TrackType, int32 TrackIndex) const
+{
+	if (TrackType == EMediaTrackType::Video) {
+		return 0;
+	}
+	return INDEX_NONE;
+}
+
+
+FString FMediaIOCorePlayerBase::GetTrackLanguage(EMediaTrackType TrackType, int32 TrackIndex) const
+{
+	return FString();
+}
+
+
+FString FMediaIOCorePlayerBase::GetTrackName(EMediaTrackType TrackType, int32 TrackIndex) const
+{
+	return FString();
+}
+
+
+bool FMediaIOCorePlayerBase::GetVideoTrackFormat(int32 TrackIndex, int32 FormatIndex, FMediaVideoTrackFormat& OutFormat) const
+{
+	if (!IsHardwareReady() || TrackIndex != 0 || FormatIndex != 0)
+	{
+		return false;
+	}
+
+	OutFormat = VideoTrackFormat;
+	return true;
+}
+
+
+bool FMediaIOCorePlayerBase::SelectTrack(EMediaTrackType TrackType, int32 TrackIndex)
+{
+	if (!IsHardwareReady() || TrackIndex < INDEX_NONE || TrackIndex != 0)
+	{
+		return false;
+	}
+
+	// Only 1 track supported
+	return (TrackType == EMediaTrackType::Audio || TrackType == EMediaTrackType::Video);
+}
+
+bool FMediaIOCorePlayerBase::SetTrackFormat(EMediaTrackType TrackType, int32 TrackIndex, int32 FormatIndex)
+{
+	return false;
+}
+
+bool FMediaIOCorePlayerBase::ReadMediaOptions(const IMediaOptions* Options)
+{
+	bUseTimeSynchronization = Options->GetMediaOption(TimeSynchronizableMedia::UseTimeSynchronizatioOption, false);
+	{
+		int32 Numerator = Options->GetMediaOption(FMediaIOCoreMediaOption::FrameRateNumerator, (int64)30);
+		int32 Denominator = Options->GetMediaOption(FMediaIOCoreMediaOption::FrameRateDenominator, (int64)1);
+		VideoFrameRate = FFrameRate(Numerator, Denominator);
+	}
+	{
+		int32 ResolutionX = Options->GetMediaOption(FMediaIOCoreMediaOption::ResolutionWidth, (int64)1920);
+		int32 ResolutionY = Options->GetMediaOption(FMediaIOCoreMediaOption::ResolutionHeight, (int64)1080);
+		VideoTrackFormat.Dim = FIntPoint(ResolutionX, ResolutionY);
+		VideoTrackFormat.FrameRates = TRange<float>(VideoFrameRate.AsDecimal());
+		VideoTrackFormat.FrameRate = VideoFrameRate.AsDecimal();
+		VideoTrackFormat.TypeName = Options->GetMediaOption(FMediaIOCoreMediaOption::VideoModeName, FString(TEXT("1080p30")));
+	}
+	return true;
+}
+
+bool FMediaIOCorePlayerBase::Exec(class UWorld* InWorld, const TCHAR* Cmd, FOutputDevice& Ar)
+{
+#if !UE_BUILD_SHIPPING
+	if (FParse::Command(&Cmd, TEXT("MediaIO")))
+	{
+		if (FParse::Command(&Cmd, TEXT("ShowInputTimecode")))
+		{
+			bIsTimecodeLogEnable = true;
+			return true;
+		}
+		else if (FParse::Command(&Cmd, TEXT("HideInputTimecode")))
+		{
+			bIsTimecodeLogEnable = false;
+			return true;
+		}
+	}
+#endif
+	return false;
+}
+
+#undef LOCTEXT_NAMESPACE
diff --git a/Engine/Plugins/Media/MediaIOFramework/Source/MediaIOCore/Private/Player/MediaIOCoreSamples.cpp b/Engine/Plugins/Media/MediaIOFramework/Source/MediaIOCore/Private/Player/MediaIOCoreSamples.cpp
new file mode 100644
--- /dev/null
+++ b/Engine/Plugins/Media/MediaIOFramework/Source/MediaIOCore/Private/Player/MediaIOCoreSamples.cpp
@@ -0,0 +1,141 @@
+// Copyright 1998-2018 Epic Games, Inc. All Rights Reserved.
+
+#include "MediaIOCoreSamples.h"
+
+#include "IMediaAudioSample.h"
+#include "IMediaBinarySample.h"
+#include "IMediaOverlaySample.h"
+#include "IMediaTextureSample.h"
+
+
+/* IMediaSamples interface
+*****************************************************************************/
+
+bool FMediaIOCoreSamples::FetchAudio(TRange<FTimespan> TimeRange, TSharedPtr<IMediaAudioSample, ESPMode::ThreadSafe>& OutSample)
+{
+	TSharedPtr<IMediaAudioSample, ESPMode::ThreadSafe> Sample;
+
+	if (!AudioSampleQueue.Peek(Sample))
+	{
+		return false;
+	}
+
+	const FTimespan SampleTime = Sample->GetTime();
+
+	if (!TimeRange.Overlaps(TRange<FTimespan>(SampleTime, SampleTime + Sample->GetDuration())))
+	{
+		return false;
+	}
+
+	AudioSampleQueue.Pop();
+	OutSample = Sample;
+
+	return true;
+}
+
+
+bool FMediaIOCoreSamples::FetchCaption(TRange<FTimespan> TimeRange, TSharedPtr<IMediaOverlaySample, ESPMode::ThreadSafe>& OutSample)
+{
+	TSharedPtr<IMediaOverlaySample, ESPMode::ThreadSafe> Sample;
+
+	if (!CaptionSampleQueue.Peek(Sample))
+	{
+		return false;
+	}
+
+	const FTimespan SampleTime = Sample->GetTime();
+
+	if (!TimeRange.Overlaps(TRange<FTimespan>(SampleTime, SampleTime + Sample->GetDuration())))
+	{
+		return false;
+	}
+
+	CaptionSampleQueue.Pop();
+	OutSample = Sample;
+
+	return true;
+}
+
+
+bool FMediaIOCoreSamples::FetchMetadata(TRange<FTimespan> TimeRange, TSharedPtr<IMediaBinarySample, ESPMode::ThreadSafe>& OutSample)
+{
+	TSharedPtr<IMediaBinarySample, ESPMode::ThreadSafe> Sample;
+
+	if (!MetadataSampleQueue.Peek(Sample))
+	{
+		return false;
+	}
+
+	const FTimespan SampleTime = Sample->GetTime();
+
+	if (!TimeRange.Overlaps(TRange<FTimespan>(SampleTime, SampleTime + Sample->GetDuration())))
+	{
+		return false;
+	}
+
+	MetadataSampleQueue.Pop();
+	OutSample = Sample;
+
+	return true;
+}
+
+
+bool FMediaIOCoreSamples::FetchSubtitle(TRange<FTimespan> TimeRange, TSharedPtr<IMediaOverlaySample, ESPMode::ThreadSafe>& OutSample)
+{
+	TSharedPtr<IMediaOverlaySample, ESPMode::ThreadSafe> Sample;
+
+	if (!SubtitleSampleQueue.Peek(Sample))
+	{
+		return false;
+	}
+
+	const FTimespan SampleTime = Sample->GetTime();
+
+	if (!TimeRange.Overlaps(TRange<FTimespan>(SampleTime, SampleTime + Sample->GetDuration())))
+	{
+		return false;
+	}
+
+	SubtitleSampleQueue.Pop();
+	OutSample = Sample;
+
+	return true;
+}
+
+
+bool FMediaIOCoreSamples::FetchVideo(TRange<FTimespan> TimeRange, TSharedPtr<IMediaTextureSample, ESPMode::ThreadSafe>& OutSample)
+{
+	TSharedPtr<IMediaTextureSample, ESPMode::ThreadSafe> Sample;
+
+	if (!VideoSampleQueue.Peek(Sample))
+	{
+		return false;
+	}
+
+	const FTimespan SampleTime = Sample->GetTime();
+
+	if (TimeRange.Overlaps(TRange<FTimespan>(SampleTime, SampleTime + Sample->GetDuration())))
+	{
+		VideoSampleQueue.Pop();
+		OutSample = Sample;
+		return true;
+	}
+
+	if (TimeRange.HasLowerBound() && SampleTime < TimeRange.GetLowerBoundValue())
+	{
+		VideoSampleQueue.Pop();
+		Sample.Reset();
+	}
+
+	return false;
+}
+
+
+void FMediaIOCoreSamples::FlushSamples()
+{
+	AudioSampleQueue.RequestFlush();
+	CaptionSampleQueue.RequestFlush();
+	MetadataSampleQueue.RequestFlush();
+	SubtitleSampleQueue.RequestFlush();
+	VideoSampleQueue.RequestFlush();
+}
diff --git a/Engine/Plugins/Media/MediaIOFramework/Source/MediaIOCore/Private/Shared/MediaIOCoreAudioSampleBase.cpp b/Engine/Plugins/Media/MediaIOFramework/Source/MediaIOCore/Private/Shared/MediaIOCoreAudioSampleBase.cpp
new file mode 100644
--- /dev/null
+++ b/Engine/Plugins/Media/MediaIOFramework/Source/MediaIOCore/Private/Shared/MediaIOCoreAudioSampleBase.cpp
@@ -0,0 +1,50 @@
+// Copyright 1998-2018 Epic Games, Inc. All Rights Reserved.
+
+#include "MediaIOCoreAudioSampleBase.h"
+
+
+FMediaIOCoreAudioSampleBase::FMediaIOCoreAudioSampleBase()
+	: Channels(0)
+	, Duration(0)
+	, SampleRate(0)
+	, Time(FTimespan::MinValue())
+{ }
+
+
+bool FMediaIOCoreAudioSampleBase::Initialize(const int32* InAudioBuffer, uint32 InBufferSize, uint32 InNumberOfChannels, uint32 InSampleRate, FTimespan InTime, const TOptional<FTimecode>& InTimecode)
+{
+	if (InAudioBuffer == nullptr || InNumberOfChannels * InSampleRate <= 0)
+	{
+		FreeSample();
+		return false;
+	}
+
+	Buffer.Reset(InBufferSize);
+	Buffer.Append(InAudioBuffer, InBufferSize);
+	Time = InTime;
+	Timecode = InTimecode;
+	Channels = InNumberOfChannels;
+	SampleRate = InSampleRate;
+	Duration = (InBufferSize * ETimespan::TicksPerSecond) / (Channels * SampleRate);
+
+	return true;
+}
+
+
+bool FMediaIOCoreAudioSampleBase::Initialize(TArray<int32> InAudioBuffer, uint32 InNumberOfChannels, uint32 InSampleRate, FTimespan InTime, const TOptional<FTimecode>& InTimecode)
+{
+	if (InAudioBuffer.Num() == 0 || InNumberOfChannels * InSampleRate <= 0)
+	{
+		FreeSample();
+		return false;
+	}
+
+	Buffer = MoveTemp(InAudioBuffer);
+	Time = InTime;
+	Timecode = InTimecode;
+	Channels = InNumberOfChannels;
+	SampleRate = InSampleRate;
+	Duration = (InAudioBuffer.Num() * ETimespan::TicksPerSecond) / (Channels * SampleRate);
+
+	return true;
+}
diff --git a/Engine/Plugins/Media/MediaIOFramework/Source/MediaIOCore/Private/Shared/MediaIOCoreBinarySampleBase.cpp b/Engine/Plugins/Media/MediaIOFramework/Source/MediaIOCore/Private/Shared/MediaIOCoreBinarySampleBase.cpp
new file mode 100644
--- /dev/null
+++ b/Engine/Plugins/Media/MediaIOFramework/Source/MediaIOCore/Private/Shared/MediaIOCoreBinarySampleBase.cpp
@@ -0,0 +1,35 @@
+// Copyright 1998-2018 Epic Games, Inc. All Rights Reserved.
+
+#include "MediaIOCoreBinarySampleBase.h"
+
+
+FMediaIOCoreBinarySampleBase::FMediaIOCoreBinarySampleBase()
+	: Time(FTimespan::Zero())
+{ }
+
+
+bool FMediaIOCoreBinarySampleBase::Initialize(const uint8* InBinaryBuffer, uint32 InBufferSize, FTimespan InTime, const TOptional<FTimecode>& InTimecode)
+{
+	if (InBinaryBuffer == nullptr)
+	{
+		Buffer.Reset();
+		return false;
+	}
+
+	Buffer.Reset(InBufferSize);
+	Buffer.Append(InBinaryBuffer, InBufferSize);
+	Time = InTime;
+	Timecode = InTimecode;
+
+	return true;
+}
+
+
+bool FMediaIOCoreBinarySampleBase::Initialize(TArray<uint8> InBinaryBuffer, FTimespan InTime, const TOptional<FTimecode>& InTimecode)
+{
+	Buffer = MoveTemp(InBinaryBuffer);
+	Time = InTime;
+	Timecode = InTimecode;
+
+	return true;
+}
diff --git a/Engine/Plugins/Media/MediaIOFramework/Source/MediaIOCore/Private/Shared/MediaIOCoreEncodeTime.cpp b/Engine/Plugins/Media/MediaIOFramework/Source/MediaIOCore/Private/Shared/MediaIOCoreEncodeTime.cpp
new file mode 100644
--- /dev/null
+++ b/Engine/Plugins/Media/MediaIOFramework/Source/MediaIOCore/Private/Shared/MediaIOCoreEncodeTime.cpp
@@ -0,0 +1,104 @@
+// Copyright 1998-2018 Epic Games, Inc. All Rights Reserved.
+
+#include "MediaIOCoreEncodeTime.h"
+
+FMediaIOCoreEncodeTime::FMediaIOCoreEncodeTime(EMediaIOCoreEncodePixelFormat InFormat, void* InBuffer, uint32 InPitch, uint32 InHeight)
+	: Format(InFormat)
+	, Buffer(InBuffer)
+	, Pitch(InPitch)
+	, Height(InHeight)
+{
+	check(Buffer);
+
+	if (Format == EMediaIOCoreEncodePixelFormat::A2B10G10R10)
+	{
+		// 10bit Encoding of colors
+		ColorBlack = 0x3;
+		ColorRed = 0xFFF;
+		ColorWhite = 0xFFFFFFFF;
+	}
+	else if (Format == EMediaIOCoreEncodePixelFormat::CharUYVY)
+	{
+		// YUVU Encoding of colors
+		ColorBlack = 0x00800080;
+		ColorRed = 0x38e4385e;
+		ColorWhite = 0xff80ff80;
+	}
+	else if (Format == EMediaIOCoreEncodePixelFormat::CharBGRA)
+	{
+		ColorBlack = FColor::Black.ToPackedARGB();
+		ColorRed = FColor::Red.ToPackedARGB();
+		ColorWhite = FColor::White.ToPackedARGB();
+	}
+	else
+	{
+		check(false);
+	}
+}
+
+void FMediaIOCoreEncodeTime::Fill(uint32 InX, uint32 InY, uint32  InWidth, uint32 InHeight, TColor InColor) const
+{
+	check(InX < Pitch);
+	check(InX + InWidth <= Pitch);
+	check(InY < Height);
+	check(InY + InHeight <= Height);
+
+	for (uint32 Line = InY; Line < InY + InHeight; ++Line)
+	{
+		for (uint32 Column = InX; Column < InX + InWidth; ++Column)
+		{
+			At(Column, Line) = InColor;
+		}
+	}
+}
+
+void FMediaIOCoreEncodeTime::FillChecker(uint32 InX, uint32 InY, uint32  InWidth, uint32 InHeight, TColor InColor0, TColor InColor1) const
+{
+	check(InX < Pitch);
+	check(InX + InWidth <= Pitch);
+	check(InY < Height);
+	check(InY + InHeight <= Height);
+
+	for (uint32 Line = InY; Line < InY + InHeight; ++Line)
+	{
+		for (uint32 Column = InX; Column < InX + InWidth; ++Column)
+		{
+			At(Column, Line) = (Column % 2) ? InColor0 : InColor1;
+		}
+	}
+}
+
+void FMediaIOCoreEncodeTime::DrawTime(uint32 InX, uint32 InY, uint32  InTime, TColor InColor) const
+{
+	uint32 Tenth = (InTime / 10);
+	uint32 Unit = (InTime % 10);
+	At(InX + Tenth, InY) = InColor;
+	At(InX + Unit, InY + 1) = InColor;
+}
+
+void FMediaIOCoreEncodeTime::Render(uint32 InX, uint32 InY, uint32 InHours, uint32 InMinutes, uint32 InSeconds, uint32 InFrames) const
+{
+	if (InX + 12 >= Pitch/sizeof(TColor) || InY + 12 >= Height)
+	{
+		return;
+	}
+
+	// clear space
+	Fill(InX, InY, 12, 12, ColorBlack);
+
+	// render pattern
+	FillChecker(InX, InY +  0,  2, 1, ColorRed, ColorBlack);
+	FillChecker(InX, InY +  1, 10, 1, ColorRed, ColorBlack);
+	FillChecker(InX, InY +  3,  6, 1, ColorRed, ColorBlack);
+	FillChecker(InX, InY +  4, 10, 1, ColorRed, ColorBlack);
+	FillChecker(InX, InY +  6,  6, 1, ColorRed, ColorBlack);
+	FillChecker(InX, InY +  7, 10, 1, ColorRed, ColorBlack);
+	FillChecker(InX, InY +  9,  6, 1, ColorRed, ColorBlack);
+	FillChecker(InX, InY + 10, 10, 1, ColorRed, ColorBlack);
+
+	// render time
+	DrawTime(InX, InY + 0, InHours, ColorWhite);
+	DrawTime(InX, InY + 3, InMinutes, ColorWhite);
+	DrawTime(InX, InY + 6, InSeconds, ColorWhite);
+	DrawTime(InX, InY + 9, InFrames, ColorWhite);
+}
diff --git a/Engine/Plugins/Media/MediaIOFramework/Source/MediaIOCore/Private/Shared/MediaIOCoreTextureSampleBase.cpp b/Engine/Plugins/Media/MediaIOFramework/Source/MediaIOCore/Private/Shared/MediaIOCoreTextureSampleBase.cpp
new file mode 100644
--- /dev/null
+++ b/Engine/Plugins/Media/MediaIOFramework/Source/MediaIOCore/Private/Shared/MediaIOCoreTextureSampleBase.cpp
@@ -0,0 +1,83 @@
+// Copyright 1998-2018 Epic Games, Inc. All Rights Reserved.
+
+#include "MediaIOCoreTextureSampleBase.h"
+
+FMediaIOCoreTextureSampleBase::FMediaIOCoreTextureSampleBase()
+	: Duration(FTimespan::Zero())
+	, SampleFormat(EMediaTextureSampleFormat::Undefined)
+	, Time(FTimespan::Zero())
+	, Stride(0)
+	, Width(0)
+	, Height(0)
+{
+}
+
+
+bool FMediaIOCoreTextureSampleBase::Initialize(const void* InVideoBuffer, uint32 InBufferSize, uint32 InStride, uint32 InWidth, uint32 InHeight, EMediaTextureSampleFormat InSampleFormat, FTimespan InTime, const TOptional<FTimecode>& InTimecode)
+{
+	FreeSample();
+
+	if ((InVideoBuffer == nullptr) || (InSampleFormat == EMediaTextureSampleFormat::Undefined))
+	{
+		return false;
+	}
+
+	Buffer.Reset(InBufferSize);
+	Buffer.Append(reinterpret_cast<const uint8*>(InVideoBuffer), InBufferSize);
+	Stride = InStride;
+	Width = InWidth;
+	Height = InHeight;
+	SampleFormat = InSampleFormat;
+	Time = InTime;
+	Timecode = InTimecode;
+
+	return true;
+}
+
+
+bool FMediaIOCoreTextureSampleBase::Initialize(TArray<uint8> InVideoBuffer, uint32 InStride, uint32 InWidth, uint32 InHeight, EMediaTextureSampleFormat InSampleFormat, FTimespan InTime, const TOptional<FTimecode>& InTimecode)
+{
+	FreeSample();
+
+	if ((InVideoBuffer.Num() == 0) || (InSampleFormat == EMediaTextureSampleFormat::Undefined))
+	{
+		return false;
+	}
+
+	Buffer = MoveTemp(InVideoBuffer);
+	Stride = InStride;
+	Width = InWidth;
+	Height = InHeight;
+	SampleFormat = InSampleFormat;
+	Time = InTime;
+	Timecode = InTimecode;
+
+	return true;
+}
+
+
+bool FMediaIOCoreTextureSampleBase::InitializeWithEvenOddLine(bool bUseEvenLine, const void* InVideoBuffer, uint32 InBufferSize, uint32 InStride, uint32 InWidth, uint32 InHeight, EMediaTextureSampleFormat InSampleFormat, FTimespan InTime, const TOptional<FTimecode>& InTimecode)
+{
+	FreeSample();
+
+	if (InVideoBuffer == nullptr || InSampleFormat == EMediaTextureSampleFormat::Undefined)
+	{
+		return false;
+	}
+
+	Buffer.Reset(InBufferSize / 2);
+	Stride = InStride;
+	Width = InWidth;
+	Height = InHeight / 2;
+	SampleFormat = InSampleFormat;
+	Timecode = InTimecode;
+	Time = InTime;
+
+	for (uint32 IndexY = (bUseEvenLine ? 0 : 1); IndexY < InHeight; IndexY += 2)
+	{
+		const uint8* Source = reinterpret_cast<const uint8*>(InVideoBuffer) + (IndexY*Stride);
+		Buffer.Append(Source, Stride);
+	}
+
+	return true;
+}
diff --git a/Engine/Plugins/Media/MediaIOFramework/Source/MediaIOCore/Private/Shared/MediaIOCoreWaitVsyncThread.cpp b/Engine/Plugins/Media/MediaIOFramework/Source/MediaIOCore/Private/Shared/MediaIOCoreWaitVsyncThread.cpp
new file mode 100644
--- /dev/null
+++ b/Engine/Plugins/Media/MediaIOFramework/Source/MediaIOCore/Private/Shared/MediaIOCoreWaitVsyncThread.cpp
@@ -0,0 +1,80 @@
+// Copyright 1998-2018 Epic Games, Inc. All Rights Reserved.
+
+#include "MediaIOCoreWaitVsyncThread.h"
+#include "IMediaIOCoreHardwareSync.h"
+
+#include "Engine/GameEngine.h"
+#include "HAL/Event.h"
+#include "HAL/RunnableThread.h"
+#include "Templates/Atomic.h"
+
+#include "MediaIOCoreModule.h"
+
+FMediaIOCoreWaitVSyncThread::FMediaIOCoreWaitVSyncThread(TSharedPtr<IMediaIOCoreHardwareSync> InHardwareSync)
+	: HardwareSync(InHardwareSync)
+	, bWaitingForSignal(false)
+	, bAlive(false)
+{
+	const bool bIsManualReset = false;
+	WaitVSync = FPlatformProcess::GetSynchEventFromPool(bIsManualReset);
+}
+
+bool FMediaIOCoreWaitVSyncThread::Init()
+{
+	bAlive.Store(true);
+	return true;
+}
+
+uint32 FMediaIOCoreWaitVSyncThread::Run()
+{
+	while (!GIsRequestingExit && bAlive.Load() && HardwareSync->IsValid())
+	{
+		// wait for event
+		HardwareSync->WaitVSync();
+
+		if (!bWaitingForSignal.Load() && bAlive)
+		{
+			UE_LOG(LogMediaIOCore, Error, TEXT("The Engine couldn't run fast enough to keep up with the VSync."));
+		}
+
+		WaitVSync->Trigger();
+	}
+
+	return 0;
+}
+
+void FMediaIOCoreWaitVSyncThread::Stop()
+{
+	bAlive.Store(false);
+}
+
+void FMediaIOCoreWaitVSyncThread::Exit()
+{
+	FPlatformProcess::ReturnSynchEventToPool(WaitVSync);
+}
+
+bool FMediaIOCoreWaitVSyncThread::Wait_GameOrRenderThread()
+{
+	bool bResult = true;
+	if (bAlive.Load())
+	{
+		bWaitingForSignal.Store(true);
+
+		int32 WaitTime = 100; //ms
+		bResult = WaitVSync->Wait(WaitTime);
+
+		//Thread could have been stopped during the wait
+		if (bAlive.Load())
+		{
+			WaitVSync->Reset();
+
+			bWaitingForSignal.Store(false);
+			if (!bResult)
+			{
+				UE_LOG(LogMediaIOCore, Error, TEXT("Lost VSync signal."));
+			}
+		}
+	}
+
+	return bResult;
+}
diff --git a/Engine/Plugins/Media/MediaIOFramework/Source/MediaIOCore/Public/IMediaIOCoreDeviceProvider.h b/Engine/Plugins/Media/MediaIOFramework/Source/MediaIOCore/Public/IMediaIOCoreDeviceProvider.h
new file mode 100644
--- /dev/null
+++ b/Engine/Plugins/Media/MediaIOFramework/Source/MediaIOCore/Public/IMediaIOCoreDeviceProvider.h
@@ -0,0 +1,37 @@
+// Copyright 1998-2018 Epic Games, Inc. All Rights Reserved.
+
+#pragma once
+
+#include "CoreMinimal.h"
+#include "MediaIOCoreDefinitions.h"
+
+class MEDIAIOCORE_API IMediaIOCoreDeviceProvider
+{
+public:
+	virtual FName GetFName() = 0;
+
+	virtual TArray<FMediaIOConnection> GetConnections() const = 0;
+	virtual TArray<FMediaIOConfiguration> GetConfigurations() const = 0;
+	virtual TArray<FMediaIOConfiguration> GetConfigurations(bool bAllowInput, bool bAllowOutput) const = 0;
+	virtual TArray<FMediaIODevice> GetDevices() const = 0;
+	virtual TArray<FMediaIOMode> GetModes(const FMediaIODevice& InDevice, bool bInOutput) const = 0;
+	virtual TArray<FMediaIOOutputConfiguration> GetOutputConfigurations() const = 0;
+
+	virtual FMediaIOConfiguration GetDefaultConfiguration() const = 0;
+	virtual FMediaIOMode GetDefaultMode() const = 0;
+	virtual FMediaIOOutputConfiguration GetDefaultOutputConfiguration() const = 0;
+
+	virtual FText ToText(const FMediaIOConfiguration& InConfiguration) const;
+	virtual FText ToText(const FMediaIOConnection& InConnection) const;
+	virtual FText ToText(const FMediaIODevice& InDevice) const;
+	virtual FText ToText(const FMediaIOMode& InMode) const;
+	virtual FText ToText(const FMediaIOOutputConfiguration& InMode) const;
+
+	static FText GetTransportName(EMediaIOSDITransportType InLinkType, EMediaIOQuadLinkTransportType InQuadLinkType);
+
+#if WITH_EDITOR
+	virtual bool ShowInputTransportInSelector() const { return true; }
+	virtual bool ShowOutputTransportInSelector() const { return true; }
+	virtual bool ShowReferenceInSelector() const { return true; }
+#endif
+};
diff --git a/Engine/Plugins/Media/MediaIOFramework/Source/MediaIOCore/Public/IMediaIOCoreHardwareSync.h b/Engine/Plugins/Media/MediaIOFramework/Source/MediaIOCore/Public/IMediaIOCoreHardwareSync.h
new file mode 100644
--- /dev/null
+++ b/Engine/Plugins/Media/MediaIOFramework/Source/MediaIOCore/Public/IMediaIOCoreHardwareSync.h
@@ -0,0 +1,21 @@
+// Copyright 1998-2018 Epic Games, Inc. All Rights Reserved.
+
+#pragma once
+
+/**
+* Interface to Hardware Sync Device
+*/
+
+class IMediaIOCoreHardwareSync
+{
+public:
+	/** Is hardware valid, and ready to synchronize. */
+	virtual bool IsValid() const = 0;
+
+	/** Wait on hardware synchronization. */
+	virtual bool WaitVSync() = 0;
+
+public:
+	/** Virtual destructor for the interface. */ 
+	virtual ~IMediaIOCoreHardwareSync() {}
+};
diff --git a/Engine/Plugins/Media/MediaIOFramework/Source/MediaIOCore/Public/IMediaIOCoreModule.h b/Engine/Plugins/Media/MediaIOFramework/Source/MediaIOCore/Public/IMediaIOCoreModule.h
new file mode 100644
--- /dev/null
+++ b/Engine/Plugins/Media/MediaIOFramework/Source/MediaIOCore/Public/IMediaIOCoreModule.h
@@ -0,0 +1,23 @@
+// Copyright 1998-2018 Epic Games, Inc. All Rights Reserved.
+#pragma once
+
+#include "CoreMinimal.h"
+#include "Modules/ModuleInterface.h"
+
+
+class IMediaIOCoreDeviceProvider;
+
+
+/**
+ * Definition the MediaIOCore module.
+ */
+class MEDIAIOCORE_API IMediaIOCoreModule : public IModuleInterface
+{
+public:
+	static IMediaIOCoreModule& Get();
+
+public:
+	virtual void RegisterDeviceProvider(IMediaIOCoreDeviceProvider* InProvider) = 0;
+	virtual void UnregisterDeviceProvider(IMediaIOCoreDeviceProvider* InProvider) = 0;
+	virtual IMediaIOCoreDeviceProvider* GetDeviceProvider(FName InProviderName) = 0;
+};
diff --git a/Engine/Plugins/Media/MediaIOFramework/Source/MediaIOCore/Public/MediaCapture.h b/Engine/Plugins/Media/MediaIOFramework/Source/MediaIOCore/Public/MediaCapture.h
new file mode 100644
--- /dev/null
+++ b/Engine/Plugins/Media/MediaIOFramework/Source/MediaIOCore/Public/MediaCapture.h
@@ -0,0 +1,200 @@
+// Copyright 1998-2018 Epic Games, Inc. All Rights Reserved.
+
+#pragma once
+
+#include "CoreTypes.h"
+#include "UObject/ObjectMacros.h"
+#include "UObject/Object.h"
+
+#include "Misc/Timecode.h"
+#include "PixelFormat.h"
+#include "RHI.h"
+#include "RHIResources.h"
+#include "HAL/CriticalSection.h"
+
+#include "MediaCapture.generated.h"
+
+class FSceneViewport;
+class UMediaOutput;
+class UTextureRenderTarget2D;
+
+/**
+ * Possible states of media capture.
+ */
+UENUM()
+enum class EMediaCaptureState
+{
+	/** Unrecoverable error occurred during capture. */
+	Error,
+
+	/** Media is currently capturing. */
+	Capturing,
+
+	/** Media is being prepared for capturing. */
+	Preparing,
+
+	/** Capture has been stopped but some frames may need to be process. */
+	StopRequested,
+
+	/** Capture has been stopped. */
+	Stopped,
+};
+
+/**
+ * Base class of additional data that can be stored for each requested capture.
+ */
+class FMediaCaptureUserData
+{};
+
+/**
+ * Abstract base class for media capture.
+ *
+ * MediaCapture capture the texture of the Render target or the SceneViewport and sends it to an external media device.
+ * MediaCapture should be created by a MediaOutput.
+ */
+UCLASS(Abstract, editinlinenew, BlueprintType, hidecategories = (Object))
+class MEDIAIOCORE_API UMediaCapture : public UObject
+{
+	GENERATED_UCLASS_BODY()
+
+public:
+	/**
+	 * Stop the actual capture if there is one.
+	 * Then start the capture of a SceneViewport.
+	 * If the SceneViewport is destroyed, the capture will stop.
+	 * The SceneViewport needs to be of the same size and have the same pixel format as requested by the media output.
+	 * @note make sure the size of the SceneViewport doesn't change during capture.
+	 * @return True if the capture was successfully started
+	 */
+	bool CaptureSceneViewport(TSharedPtr<FSceneViewport>& SceneViewport);
+
+	/**
+	 * Stop the current capture if there is one.
+	 * Then find and capture every frame from active SceneViewport.
+	 * It can only find a SceneViewport when you play in Standalone or in "New Editor Window PIE".
+	 * If the active SceneViewport is destroyed, the capture will stop.
+	 * The SceneViewport needs to be of the same size and have the same pixel format as requested by the media output.
+	 * @return True if the capture was successfully started
+	 */
+	UFUNCTION(BlueprintCallable, Category = "Media|Output")
+	bool CaptureActiveSceneViewport();
+
+	/**
+	 * Stop the actual capture if there is one.
+	 * Then capture every frame for a TextureRenderTarget2D.
+	 * The TextureRenderTarget2D needs to be of the same size and have the same pixel format as requested by the media output.
+	 * @return True if the capture was successfully started
+	 */
+	UFUNCTION(BlueprintCallable, Category = "Media|Output")
+	bool CaptureTextureRenderTarget2D(UTextureRenderTarget2D* RenderTarget);
+
+	/**
+	 * Update the current capture with a SceneViewport.
+	 * If the SceneViewport is destroyed, the capture will stop.
+	 * The SceneViewport needs to be of the same size and have the same pixel format as requested by the media output.
+	 * @note make sure the size of the SceneViewport doesn't change during capture.
+	 * @return Return true if the capture was successfully updated. If false is returned, the capture was stopped.
+	 */
+	bool UpdateSceneViewport(TSharedPtr<FSceneViewport>& SceneViewport);
+
+
+	/**
+	 * Update the current capture with every frame for a TextureRenderTarget2D.
+	 * The TextureRenderTarget2D needs to be of the same size and have the same pixel format as requested by the media output.
+	 * @return Return true if the capture was successfully updated. If false is returned, the capture was stopped.
+	 */
+	UFUNCTION(BlueprintCallable, Category = "Media|Output")
+	bool UpdateTextureRenderTarget2D(UTextureRenderTarget2D* RenderTarget);
+
+
+	/**
+	 * Stop the previous requested capture.
+	 * @param bAllowPendingFrameToBeProcess	Keep copying the pending frames asynchronously or stop immediately without copying the pending frames.
+	 */
+	UFUNCTION(BlueprintCallable, Category = "Media|Output")
+	void StopCapture(bool bAllowPendingFrameToBeProcess);
+
+	/** Get the current state of the capture. */
+	UFUNCTION(BlueprintCallable, Category = "Media|Output")
+	virtual EMediaCaptureState GetState() const { return MediaState; }
+
+	/** Set the media output. Can only be set when the capture is stopped. */
+	UFUNCTION(BlueprintCallable, Category = "Media|Output")
+	void SetMediaOutput(UMediaOutput* InMediaOutput);
+
+	/** Get the desired size of the current capture. */
+	UFUNCTION(BlueprintCallable, Category = "Media|Output")
+	FIntPoint GetDesiredSize() const { return DesiredSize; }
+
+	/** Get the desired pixel format of the current capture. */
+	UFUNCTION(BlueprintCallable, Category = "Media|Output")
+	EPixelFormat GetDesiredPixelFormat() const { return DesiredPixelFormat; }
+
+	/** Check whether this capture has any processing left to do. */
+	virtual bool HasFinishedProcessing() const;
+
+public:
+	//~ UObject interface
+	virtual void BeginDestroy() override;
+	virtual FString GetDesc() override;
+
+protected:
+	//~ UMediaCapture interface
+	virtual bool ValidateMediaOutput() const;
+	virtual bool CaptureSceneViewportImpl(TSharedPtr<FSceneViewport>& InSceneViewport) { return true; }
+	virtual bool CaptureRenderTargetImpl(UTextureRenderTarget2D* InRenderTarget) { return true; }
+	virtual bool UpdateSceneViewportImpl(TSharedPtr<FSceneViewport>& InSceneViewport) { return true; }
+	virtual bool UpdateRenderTargetImpl(UTextureRenderTarget2D* InRenderTarget) { return true; }
+	virtual void StopCaptureImpl(bool bAllowPendingFrameToBeProcess) { }
+
+
+	struct FCaptureBaseData
+	{
+		FCaptureBaseData();
+
+		FTimecode SourceFrameTimecode;
+		uint32 SourceFrameNumberRenderThread;
+	};
+	virtual TSharedPtr<FMediaCaptureUserData> GetCaptureFrameUserData_GameThread() { return TSharedPtr<FMediaCaptureUserData>(); }
+	virtual void OnFrameCaptured_RenderingThread(const FCaptureBaseData& InBaseData, TSharedPtr<FMediaCaptureUserData> InUserData, void* InBuffer, int32 Width, int32 Height) { }
+
+protected:
+	UTextureRenderTarget2D* GetTextureRenderTarget() { return CapturingRenderTarget; }
+	TSharedPtr<FSceneViewport> GetCapturingSceneViewport() { return CapturingSceneViewport.Pin(); }
+
+protected:
+	UPROPERTY(Transient)
+	UMediaOutput* MediaOutput;
+
+	EMediaCaptureState MediaState;
+
+private:
+	void InitializeResolveTarget(int32 InNumberOfBuffers);
+	void OnEndFrame_GameThread();
+
+private:
+	struct FCaptureFrame
+	{
+		FCaptureFrame();
+
+		FTexture2DRHIRef ReadbackTexture;
+		FCaptureBaseData CaptureBaseData;
+		bool bResolvedTargetRequested;
+		TSharedPtr<FMediaCaptureUserData> UserData;
+	};
+
+	TArray<FCaptureFrame> CaptureFrames;
+	int32 CurrentResolvedTargetIndex;
+	int32 NumberOfCaptureFrame;
+
+	UPROPERTY(Transient)
+	UTextureRenderTarget2D* CapturingRenderTarget;
+	TWeakPtr<FSceneViewport> CapturingSceneViewport;
+	FCriticalSection AccessingCapturingSource;
+
+	FIntPoint DesiredSize;
+	EPixelFormat DesiredPixelFormat;
+
+	bool bResolvedTargetInitialized;
+	bool bWaitingForResolveCommandExecution;
+};
diff --git a/Engine/Plugins/Media/MediaIOFramework/Source/MediaIOCore/Public/MediaIOCoreAudioSampleBase.h b/Engine/Plugins/Media/MediaIOFramework/Source/MediaIOCore/Public/MediaIOCoreAudioSampleBase.h
new file mode 100644
--- /dev/null
+++ b/Engine/Plugins/Media/MediaIOFramework/Source/MediaIOCore/Public/MediaIOCoreAudioSampleBase.h
@@ -0,0 +1,120 @@
+// Copyright 1998-2018 Epic Games, Inc. All Rights Reserved.
+
+#pragma once
+
+#include "IMediaAudioSample.h"
+#include "MediaObjectPool.h"
+
+/*
+ * Implements a media audio sample.
+ */
+class MEDIAIOCORE_API FMediaIOCoreAudioSampleBase
+	: public IMediaAudioSample
+	, public IMediaPoolable
+{
+public:
+
+	/** Default constructor. */
+	FMediaIOCoreAudioSampleBase();
+
+	/**
+	 * Initialize the sample.
+	 *
+	 * @param InAudioBuffer The audio frame data.
+	 * @param InBufferSize The size of the audio buffer.
+	 * @param InNumberOfChannels The number of channel of the audio buffer.
+	 * @param InSampleRate The sample rate of the audio buffer.
+	 * @param InTime The sample time (in the player's own clock).
+	 * @param InTimecode The sample timecode if available.
+	 */
+	bool Initialize(const int32* InAudioBuffer, uint32 InBufferSize, uint32 InNumberOfChannels, uint32 InSampleRate, FTimespan InTime, const TOptional<FTimecode>& InTimecode);
+
+	/**
+	 * Initialize the sample.
+	 *
+	 * @param InBinaryBuffer The metadata frame data.
+	 * @param InNumberOfChannels The number of channel of the audio buffer.
+	 * @param InSampleRate The sample rate of the audio buffer.
+	 * @param InTime The sample time (in the player's own clock).
+	 * @param InTimecode The sample timecode if available.
+	 */
+	bool Initialize(TArray<int32> InAudioBuffer, uint32 InNumberOfChannels, uint32 InSampleRate, FTimespan InTime, const TOptional<FTimecode>& InTimecode);
+
+public:
+
+	//~ IMediaAudioSample interface
+
+	virtual const void* GetBuffer() override
+	{
+		return Buffer.GetData();
+	}
+
+	virtual uint32 GetChannels() const override
+	{
+		return Channels;
+	}
+
+	virtual FTimespan GetDuration() const override
+	{
+		return Duration;
+	}
+
+	virtual EMediaAudioSampleFormat GetFormat() const override
+	{
+		return EMediaAudioSampleFormat::Int32;
+	}
+
+	virtual uint32 GetFrames() const override
+	{
+		return Buffer.Num() / Channels;
+	}
+
+	virtual uint32 GetSampleRate() const override
+	{
+		return SampleRate;
+	}
+
+	virtual FTimespan GetTime() const override
+	{
+		return Time;
+	}
+
+	virtual TOptional<FTimecode> GetTimecode() const override
+	{
+		return Timecode;
+	}
+
+public:
+
+	//~ IMediaPoolable interface
+
+	virtual void ShutdownPoolable() override
+	{
+		FreeSample();
+	}
+
+protected:
+	virtual void FreeSample()
+	{
+		Buffer.Reset();
+	}
+
+protected:
+	/** The sample's frame buffer. */
+	TArray<int32> Buffer;
+
+	/** Number of audio channels. */
+	uint32 Channels;
+
+	/** The duration for which the sample is valid. */
+	FTimespan Duration;
+
+	/** Audio sample rate (in samples per second). */
+	uint32 SampleRate;
+
+	/** Sample time. */
+	FTimespan Time;
+
+	/** Sample timecode. */
+	TOptional<FTimecode> Timecode;
+};
diff --git a/Engine/Plugins/Media/MediaIOFramework/Source/MediaIOCore/Public/MediaIOCoreBinarySampleBase.h b/Engine/Plugins/Media/MediaIOFramework/Source/MediaIOCore/Public/MediaIOCoreBinarySampleBase.h
new file mode 100644
--- /dev/null
+++ b/Engine/Plugins/Media/MediaIOFramework/Source/MediaIOCore/Public/MediaIOCoreBinarySampleBase.h
@@ -0,0 +1,94 @@
+// Copyright 1998-2018 Epic Games, Inc. All Rights Reserved.
+
+#pragma once
+
+#include "IMediaBinarySample.h"
+#include "MediaObjectPool.h"
+
+/**
+ * Implements a media binary data sample for AjaMedia.
+ */
+class MEDIAIOCORE_API FMediaIOCoreBinarySampleBase
+	: public IMediaBinarySample
+	, public IMediaPoolable
+{
+public:
+
+	/** Default constructor. */
+	FMediaIOCoreBinarySampleBase();
+
+	/**
+	 * Initialize the sample.
+	 *
+	 * @param InBinaryBuffer The metadata frame data.
+	 * @param InBufferSize The size of the InBinaryBuffer.
+	 * @param InTimecode The sample timecode if available.
+	 * @param InTime The sample time (in the player's own clock).
+	 */
+	bool Initialize(const uint8* InBinaryBuffer, uint32 InBufferSize, FTimespan InTime, const TOptional<FTimecode>& InTimecode);
+
+	/**
+	 * Initialize the sample.
+	 *
+	 * @param InBinaryBuffer The metadata frame data.
+	 * @param InTimecode The sample timecode if available.
+	 * @param InTime The sample time (in the player's own clock).
+	 */
+	bool Initialize(TArray<uint8> InBinaryBuffer, FTimespan InTime, const TOptional<FTimecode>& InTimecode);
+
+public:
+
+	//~ IMediaBinarySample interface
+
+	virtual const void* GetData() override
+	{
+		return Buffer.GetData();
+	}
+
+	virtual FTimespan GetDuration() const override
+	{
+		return FTimespan::Zero();
+	}
+
+	virtual uint32 GetSize() const override
+	{
+		return Buffer.Num();
+	}
+
+	virtual FTimespan GetTime() const override
+	{
+		return Time;
+	}
+
+	virtual TOptional<FTimecode> GetTimecode() const override
+	{
+		return Timecode;
+	}
+
+public:
+
+	//~ IMediaPoolable interface
+
+	virtual void ShutdownPoolable() override
+	{
+		FreeSample();
+	}
+
+protected:
+
+	virtual void FreeSample()
+	{
+		Buffer.Reset();
+	}
+
+protected:
+
+	/** The sample's frame buffer. */
+	TArray<uint8> Buffer;
+
+	/** Sample time. */
+	FTimespan Time;
+
+	/** Sample timecode. */
+	TOptional<FTimecode> Timecode;
+};
diff --git a/Engine/Plugins/Media/MediaIOFramework/Source/MediaIOCore/Public/MediaIOCoreCommonDisplayMode.h b/Engine/Plugins/Media/MediaIOFramework/Source/MediaIOCore/Public/MediaIOCoreCommonDisplayMode.h
new file mode 100644
--- /dev/null
+++ b/Engine/Plugins/Media/MediaIOFramework/Source/MediaIOCore/Public/MediaIOCoreCommonDisplayMode.h
@@ -0,0 +1,37 @@
+// Copyright 1998-2018 Epic Games, Inc. All Rights Reserved.
+
+#pragma once
+
+#include "CoreMinimal.h"
+#include "Containers/ArrayView.h"
+#include "Misc/FrameRate.h"
+
+struct FMediaIOCommonDisplayModeInfo
+{
+	int32 Width;
+	int32 Height;
+	FFrameRate FrameRate;
+	int32 FieldPerFrame;
+	FText Name;
+};
+
+struct FMediaIOCommonDisplayModeResolutionInfo
+{
+	int32 Width;
+	int32 Height;
+	FText Name;
+};
+
+struct MEDIAIOCORE_API FMediaIOCommonDisplayModes
+{
+	static TArrayView<const FMediaIOCommonDisplayModeInfo> GetAllModes();
+	static TArrayView<const FMediaIOCommonDisplayModeResolutionInfo> GetAllResolutions();
+
+	static const FMediaIOCommonDisplayModeResolutionInfo* Find(int32 InWidth, int32 InHeight);
+	static const FMediaIOCommonDisplayModeInfo* Find(int32 InWidth, int32 InHeight, const FFrameRate& InFrameRate, int32 InFieldPerFrame);
+	static bool Contains(int32 InWidth, int32 InHeight) { return Find(InWidth, InHeight) != nullptr; }
+	static bool Contains(int32 InWidth, int32 InHeight, const FFrameRate& InFrameRate, int32 InFieldPerFrame) { return Find(InWidth, InHeight, InFrameRate, InFieldPerFrame) != nullptr; }
+
+	static FText GetMediaIOCommonDisplayModeResolutionInfoName(int32 InWidth, int32 InHeight);
+	static FText GetMediaIOCommonDisplayModeInfoName(int32 InWidth, int32 InHeight, const FFrameRate& InFrameRate, int32 InFieldPerFrame);
+};
\ No newline at end of file
diff --git a/Engine/Plugins/Media/MediaIOFramework/Source/MediaIOCore/Public/MediaIOCoreDefinitions.h b/Engine/Plugins/Media/MediaIOFramework/Source/MediaIOCore/Public/MediaIOCoreDefinitions.h
new file mode 100644
--- /dev/null
+++ b/Engine/Plugins/Media/MediaIOFramework/Source/MediaIOCore/Public/MediaIOCoreDefinitions.h
@@ -0,0 +1,262 @@
+// Copyright 1998-2018 Epic Games, Inc. All Rights Reserved.
+
+#pragma once
+
+#include "CoreMinimal.h"
+#include "Misc/FrameRate.h"
+#include "UObject/ObjectMacros.h"
+#include "MediaIOCoreDefinitions.generated.h"
+
+
+/**
+ * SDI transport type.
+ */
+UENUM()
+enum class EMediaIOSDITransportType
+{
+	SingleLink,
+	DualLink,
+	QuadLink,
+};
+
+
+/**
+ * Quad link transport type.
+ */
+UENUM()
+enum class EMediaIOQuadLinkTransportType
+{
+	SquareDivision,
+	TwoSampleInterleave,
+};
+
+
+/**
+ * SDI transport type.
+ */
+UENUM()
+enum class EMediaIOStandardType
+{
+	Progressive,
+	Interlaced,
+	ProgressiveSegmentedFrame,
+};
+
+/**
+ * Timecode formats.
+ */
+UENUM()
+enum class EMediaIOTimecodeFormat
+{
+	None,
+	LTC,
+	VITC,
+};
+
+
+/**
+ * SDI Output type.
+ */
+UENUM()
+enum class EMediaIOOutputType : uint8
+{
+	Fill			UMETA(DisplayName="Fill"),
+	FillAndKey		UMETA(DisplayName="Fill & Key"),
+};
+
+
+/**
+ * SDI Output type.
+ */
+UENUM()
+enum class EMediaIOReferenceType
+{
+	FreeRun,
+	External,
+	Input
+};
+
+
+/**
+ * Identifies a device.
+ */
+USTRUCT()
+struct MEDIAIOCORE_API FMediaIODevice
+{
+	GENERATED_BODY()
+
+public:
+	FMediaIODevice();
+
+	/** The retail/display name of the Device. */
+	UPROPERTY(EditAnywhere, Category=Device)
+	FName DeviceName;
+
+	/** The device identifier. */
+	UPROPERTY(EditAnywhere, Category=Device)
+	int32 DeviceIdentifier;
+
+public:
+	bool operator==(const FMediaIODevice& Other) const;
+
+	/** Return true if the device has been set properly */
+	bool IsValid() const;
+};
+
+
+/**
+ * Identifies an media connection.
+ */
+USTRUCT()
+struct MEDIAIOCORE_API FMediaIOConnection
+{
+	GENERATED_BODY()
+
+public:
+	FMediaIOConnection();
+
+	/** The device identifier. */
+	UPROPERTY(EditAnywhere, Category=Connection)
+	FMediaIODevice Device;
+	
+	/** The protocol used by the MediaFramework. */
+	UPROPERTY(EditAnywhere, Category=Connection)
+	FName Protocol;
+
+	/** The type of cable link used for that configuration */
+	UPROPERTY(EditAnywhere, Category=Connection)
+	EMediaIOSDITransportType TransportType;
+
+	/** The type of cable link used for that configuration */
+	UPROPERTY(EditAnywhere, Category=Connection)
+	EMediaIOQuadLinkTransportType QuadTransportType;
+
+	/** The port of the video channel on the device. */
+	UPROPERTY(EditAnywhere, Category=Connection)
+	int32 PortIdentifier;
+
+public:
+	bool operator==(const FMediaIOConnection& Other) const;
+
+	/**
+	 * Get a url used by the media framework
+	 * @return Url representation, "protocol://device0/single1"
+	 */
+	FString ToUrl() const;
+
+	/** Return true if the connection has been set properly */
+	bool IsValid() const;
+};
+
+
+/**
+ * Identifies a media mode.
+ */
+USTRUCT()
+struct MEDIAIOCORE_API FMediaIOMode
+{
+	GENERATED_BODY()
+
+public:
+	FMediaIOMode();
+
+public:
+	/** The mode's frame rate. */
+	UPROPERTY(VisibleAnywhere, Category=Mode)
+	FFrameRate FrameRate;
+
+	/** The mode's image resolution. */
+	UPROPERTY(VisibleAnywhere, Category=Mode)
+	FIntPoint Resolution;
+
+	/** The mode's scanning type. */
+	UPROPERTY(VisibleAnywhere, Category=Mode)
+	EMediaIOStandardType Standard;
+
+	/** The mode's identifier for the device. */
+	UPROPERTY(VisibleAnywhere, Category=Mode)
+	int32 DeviceModeIdentifier;
+
+public:
+	bool operator==(const FMediaIOMode& Other) const;
+
+	/** Return the name of this mode from MediaIOCommonDisplayModes. */
+	FText GetModeName() const;
+
+	/** Return true if the mode has been set properly. */
+	bool IsValid() const;
+};
+
+
+/**
+ * Configuration of a device input / output.
+ */
+USTRUCT()
+struct MEDIAIOCORE_API FMediaIOConfiguration
+{
+	GENERATED_BODY()
+	
+public:
+	FMediaIOConfiguration();
+
+public:
+	/** Configured as an input or output. */
+	UPROPERTY()
+	bool bIsInput;
+
+	/** The configuration's device and transport type. */
+	UPROPERTY(VisibleAnywhere, Category=Configuration)
+	FMediaIOConnection MediaConnection;
+
+	/** The configuration's video mode. */
+	UPROPERTY(VisibleAnywhere, Category=Configuration)
+	FMediaIOMode MediaMode;
+
+public:
+	bool operator== (const FMediaIOConfiguration& Other) const;
+
+	/** Return true if the configuration has been set properly */
+	bool IsValid() const;
+};
+
+/**
+ * Configuration of a device output.
+ */
+USTRUCT()
+struct MEDIAIOCORE_API FMediaIOOutputConfiguration
+{
+	GENERATED_BODY()
+
+public:
+	FMediaIOOutputConfiguration();
+
+	/** The signal output format. */
+	UPROPERTY(VisibleAnywhere, Category=Configuration)
+	FMediaIOConfiguration MediaConfiguration;
+
+	/** Whether to output the fill or the fill and key. */
+	UPROPERTY(VisibleAnywhere, Category=Configuration)
+	EMediaIOOutputType OutputType;
+
+	/**
+	 * The port of the video channel on the device to output the key to.
+	 * @note	'Frame Buffer Pixel Format' must be set to at least 8 bits of alpha.
+	 * @note	'Enable alpha channel support in post-processing' must be set to 'Allow through tonemapper'.
+	 */
+	UPROPERTY(VisibleAnywhere, Category=Configuration)
+	int32 KeyPortIdentifier;
+
+	/** The Device output sync with either its internal clock, an external reference, or an other input. */
+	UPROPERTY(VisibleAnywhere, Category=Configuration)
+	EMediaIOReferenceType OutputReference;
+
+	/** The port of the video channel on the device to output the synchronize to. */
+	UPROPERTY(VisibleAnywhere, Category=Configuration)
+	int32 ReferencePortIdentifier;
+
+public:
+	bool operator== (const FMediaIOOutputConfiguration& Other) const;
+
+	/** Return true if the configuration has been set properly */
+	bool IsValid() const;
+};
diff --git a/Engine/Plugins/Media/MediaIOFramework/Source/MediaIOCore/Public/MediaIOCoreEncodeTime.h b/Engine/Plugins/Media/MediaIOFramework/Source/MediaIOCore/Public/MediaIOCoreEncodeTime.h
new file mode 100644
--- /dev/null
+++ b/Engine/Plugins/Media/MediaIOFramework/Source/MediaIOCore/Public/MediaIOCoreEncodeTime.h
@@ -0,0 +1,49 @@
+// Copyright 1998-2018 Epic Games, Inc. All Rights Reserved.
+
+#pragma once
+
+#include "CoreMinimal.h"
+#include "IMediaTextureSample.h"
+
+enum class MEDIAIOCORE_API EMediaIOCoreEncodePixelFormat
+{
+	A2B10G10R10,
+	CharBGRA,
+	CharUYVY,
+};
+
+class MEDIAIOCORE_API FMediaIOCoreEncodeTime
+{
+public:
+	FMediaIOCoreEncodeTime(EMediaIOCoreEncodePixelFormat InFormat, void* InBuffer, uint32 InPitch, uint32 InHeight);
+	void Render(uint32 InX, uint32 InY, uint32 InHours, uint32 InMinutes, uint32 InSeconds, uint32 InFrames) const;
+
+protected:
+	using TColor = int32;
+	void Fill(uint32 InX, uint32 InY, uint32 InWidth, uint32 InHeight, TColor InColor) const;
+	void FillChecker(uint32 InX, uint32 InY, uint32 InWidth, uint32 InHeight, TColor InColor0, TColor InColor1) const;
+	void DrawTime(uint32 InX, uint32 InY, uint32 inTime, TColor InColor) const;
+
+	inline TColor& At(uint32 InX, uint32 InY) const
+	{
+		return *(reinterpret_cast<TColor*>(reinterpret_cast<char*>(Buffer) + (Pitch * InY)) + InX);
+	}
+
+protected:
+	/** Pixel format */
+	EMediaIOCoreEncodePixelFormat Format;
+
+	/** Pointer to pixels */
+	void* Buffer;
+
+	/** Pitch of image */
+	uint32 Pitch;
+
+	/** Height of image */
+	uint32 Height;
+
+protected:
+	TColor ColorBlack;
+	TColor ColorRed;
+	TColor ColorWhite;
+};
\ No newline at end of file
diff --git a/Engine/Plugins/Media/MediaIOFramework/Source/MediaIOCore/Public/MediaIOCorePlayerBase.h b/Engine/Plugins/Media/MediaIOFramework/Source/MediaIOCore/Public/MediaIOCorePlayerBase.h
new file mode 100644
--- /dev/null
+++ b/Engine/Plugins/Media/MediaIOFramework/Source/MediaIOCore/Public/MediaIOCorePlayerBase.h
@@ -0,0 +1,175 @@
+// Copyright 1998-2018 Epic Games, Inc. All Rights Reserved.
+
+#pragma once
+
+#include "CoreMinimal.h"
+
+#include "IMediaCache.h"
+#include "IMediaControls.h"
+#include "IMediaPlayer.h"
+#include "IMediaTracks.h"
+#include "IMediaView.h"
+
+#include "HAL/CriticalSection.h"
+#include "Misc/CoreMisc.h"
+
+#include "Misc/FrameRate.h"
+
+class FMediaIOCoreSamples;
+class IMediaEventSink;
+
+enum class EMediaTextureSampleFormat;
+
+struct MEDIAIOCORE_API FMediaIOCoreMediaOption
+{
+	static const FName FrameRateNumerator;
+	static const FName FrameRateDenominator;
+	static const FName ResolutionWidth;
+	static const FName ResolutionHeight;
+	static const FName VideoModeName;
+};
+
+
+/**
+ * Implements a base player for hardware IO cards. 
+ *
+ * The processing of metadata and video frames is delayed until the fetch stage
+ * (TickFetch) in order to increase the window of opportunity for receiving
+ * frames for the current render frame time code.
+ *
+ * Depending on whether the media source enables time code synchronization,
+ * the player's current play time (CurrentTime) is derived either from the
+ * time codes embedded in frames or from the Engine's global time code.
+ */
+class MEDIAIOCORE_API FMediaIOCorePlayerBase
+	: public IMediaPlayer
+	, protected IMediaCache
+	, protected IMediaControls
+	, protected IMediaTracks
+	, protected IMediaView
+	, public FSelfRegisteringExec
+{
+public:
+
+	/**
+	 * Create and initialize a new instance.
+	 *
+	 * @param InEventSink The object that receives media events from this player.
+	 */
+	FMediaIOCorePlayerBase(IMediaEventSink& InEventSink);
+
+	/** Virtual destructor. */
+	virtual ~FMediaIOCorePlayerBase();
+
+public:
+
+	//~ IMediaPlayer interface
+
+	virtual void Close() override;
+	virtual bool Open(const FString& Url, const IMediaOptions* Options) override;
+	virtual bool Open(const TSharedRef<FArchive, ESPMode::ThreadSafe>& Archive, const FString& OriginalUrl, const IMediaOptions* Options) override;
+
+	FString GetInfo() const;
+	virtual IMediaCache& GetCache() override;
+	virtual IMediaControls& GetControls() override;
+	virtual IMediaSamples& GetSamples() override;
+	const FMediaIOCoreSamples& GetSamples() const;
+	virtual FString GetStats() const override;
+	virtual IMediaTracks& GetTracks() override;
+	virtual FString GetUrl() const override;
+	virtual IMediaView& GetView() override;
+	virtual void TickTimeManagement();
+
+public:
+	//~ IMediaCache interface
+	
+	virtual bool QueryCacheState(EMediaCacheState State, TRangeSet<FTimespan>& OutTimeRanges) const override;
+	virtual int32 GetSampleCount(EMediaCacheState State) const override;
+
+protected:
+
+	//~ IMediaControls interface
+
+	virtual bool CanControl(EMediaControl Control) const override;
+	virtual FTimespan GetDuration() const override;
+	virtual float GetRate() const override;
+	virtual EMediaState GetState() const override;
+	virtual EMediaStatus GetStatus() const override;
+	virtual TRangeSet<float> GetSupportedRates(EMediaRateThinning Thinning) const override;
+	virtual FTimespan GetTime() const override;
+	virtual bool IsLooping() const override;
+	virtual bool Seek(const FTimespan& Time) override;
+	virtual bool SetLooping(bool Looping) override;
+	virtual bool SetRate(float Rate) override;
+
+protected:
+
+	//~ IMediaTracks interface
+
+	virtual bool GetAudioTrackFormat(int32 TrackIndex, int32 FormatIndex, FMediaAudioTrackFormat& OutFormat) const override;
+	virtual int32 GetNumTracks(EMediaTrackType TrackType) const override;
+	virtual int32 GetNumTrackFormats(EMediaTrackType TrackType, int32 TrackIndex) const override;
+	virtual int32 GetSelectedTrack(EMediaTrackType TrackType) const override;
+	virtual FText GetTrackDisplayName(EMediaTrackType TrackType, int32 TrackIndex) const override;
+	virtual int32 GetTrackFormat(EMediaTrackType TrackType, int32 TrackIndex) const override;
+	virtual FString GetTrackLanguage(EMediaTrackType TrackType, int32 TrackIndex) const override;
+	virtual FString GetTrackName(EMediaTrackType TrackType, int32 TrackIndex) const override;
+	virtual bool GetVideoTrackFormat(int32 TrackIndex, int32 FormatIndex, FMediaVideoTrackFormat& OutFormat) const override;
+	virtual bool SelectTrack(EMediaTrackType TrackType, int32 TrackIndex) override;
+	virtual bool SetTrackFormat(EMediaTrackType TrackType, int32 TrackIndex, int32 FormatIndex) override;
+
+
+protected:
+
+	//~ FSelfRegisteringExec
+	
+	virtual bool Exec(class UWorld* InWorld, const TCHAR* Cmd, FOutputDevice& Ar) override;
+
+protected:
+	virtual bool IsHardwareReady() const = 0;
+
+	/** Return true if the options combination are valid */
+	virtual bool ReadMediaOptions(const IMediaOptions* Options);
+
+protected:
+	/** Critical section for synchronizing access to receiver and sinks. */
+	FCriticalSection CriticalSection;
+
+	/** Enable timecode logging */
+	bool bIsTimecodeLogEnable;
+
+	/** Url used to open the media player */
+	FString OpenUrl;
+
+	/** format of the video */
+	FMediaVideoTrackFormat VideoTrackFormat;
+	
+	/** format of the audio */
+	FMediaAudioTrackFormat AudioTrackFormat;
+
+	/** Current state of the media player. */
+	EMediaState CurrentState;
+
+	/** Current playback time. */
+	FTimespan CurrentTime;
+
+	/** The media event handler. */
+	IMediaEventSink& EventSink;
+
+	/** Video frame rate in the last received sample. */
+	FFrameRate VideoFrameRate;
+	
+	/** Number of frames drop from the last tick. */
+	uint32 LastFrameDropCount;
+
+	/** The media sample cache. */
+	FMediaIOCoreSamples* Samples;
+
+	/** Whether to use the Synchronization Time module as time source. */
+	bool bUseTimeSynchronization;
+
+	/** Previous frame Timespan */
+	FTimespan PreviousFrameTimespan;
+};
+
+
diff --git a/Engine/Plugins/Media/MediaIOFramework/Source/MediaIOCore/Public/MediaIOCoreSamples.h b/Engine/Plugins/Media/MediaIOFramework/Source/MediaIOCore/Public/MediaIOCoreSamples.h
new file mode 100644
--- /dev/null
+++ b/Engine/Plugins/Media/MediaIOFramework/Source/MediaIOCore/Public/MediaIOCoreSamples.h
@@ -0,0 +1,244 @@
+// Copyright 1998-2018 Epic Games, Inc. All Rights Reserved.
+
+#pragma once
+
+#include "CoreTypes.h"
+#include "IMediaSamples.h"
+#include "Templates/SharedPointer.h"
+
+#include "MediaSampleQueue.h"
+
+class IMediaAudioSample;
+class IMediaBinarySample;
+class IMediaOverlaySample;
+class IMediaTextureSample;
+
+/**
+ * General purpose media sample queue.
+ */
+class MEDIAIOCORE_API FMediaIOCoreSamples
+	: public IMediaSamples
+{
+public:
+	FMediaIOCoreSamples() = default;
+	FMediaIOCoreSamples(const FMediaIOCoreSamples&) = delete;
+	FMediaIOCoreSamples& operator=(const FMediaIOCoreSamples&) = delete;
+
+public:
+
+	/**
+	 * Add the given audio sample to the cache.
+	 *
+	 * @param Sample The sample to add.
+	 * @return True if the operation succeeds.
+	 * @see AddCaption, AddMetadata, AddSubtitle, AddVideo, PopAudio, NumAudioSamples
+	 */
+	bool AddAudio(const TSharedRef<IMediaAudioSample, ESPMode::ThreadSafe>& Sample)
+	{
+		return AudioSampleQueue.Enqueue(Sample);
+	}
+
+	/**
+	 * Add the given caption sample to the cache.
+	 *
+	 * @param Sample The sample to add.
+	 * @return True if the operation succeeds.
+	 * @see AddAudio, AddMetadata, AddSubtitle, AddVideo, PopCaption, NumCaptionSamples
+	 */
+	bool AddCaption(const TSharedRef<IMediaOverlaySample, ESPMode::ThreadSafe>& Sample)
+	{
+		return CaptionSampleQueue.Enqueue(Sample);
+	}
+
+	/**
+	 * Add the given metadata sample to the cache.
+	 *
+	 * @param Sample The sample to add.
+	 * @return True if the operation succeeds.
+	 * @see AddAudio, AddCaption, AddSubtitle, AddVideo, PopMetadata, NumMetadataSamples
+	 */
+	bool AddMetadata(const TSharedRef<IMediaBinarySample, ESPMode::ThreadSafe>& Sample)
+	{
+		return MetadataSampleQueue.Enqueue(Sample);
+	}
+
+	/**
+	 * Add the given subtitle sample to the cache.
+	 *
+	 * @param Sample The sample to add.
+	 * @return True if the operation succeeds.
+	 * @see AddAudio, AddCaption, AddMetadata, AddVideo, PopSubtitle, NumSubtitleSamples
+	 */
+	bool AddSubtitle(const TSharedRef<IMediaOverlaySample, ESPMode::ThreadSafe>& Sample)
+	{
+		return SubtitleSampleQueue.Enqueue(Sample);
+	}
+
+	/**
+	 * Add the given video sample to the cache.
+	 *
+	 * @param Sample The sample to add.
+	 * @return True if the operation succeeds.
+	 * @see AddAudio, AddCaption, AddMetadata, AddSubtitle, PopVideo, NumVideoSamples
+	 */
+	bool AddVideo(const TSharedRef<IMediaTextureSample, ESPMode::ThreadSafe>& Sample)
+	{
+		return VideoSampleQueue.Enqueue(Sample);
+	}
+
+	/**
+	 * Pop a Audio sample from the cache.
+	 *
+	 * @return True if the operation succeeds.
+	 * @see AddAudio, NumAudioSamples
+	 */
+	bool PopAudio()
+	{
+		return AudioSampleQueue.Pop();
+	}
+
+	/**
+	 * Pop a Caption sample from the cache.
+	 *
+	 * @return True if the operation succeeds.
+	 * @see AddCaption, NumCaption
+	 */
+	bool PopCaption()
+	{
+		return CaptionSampleQueue.Pop();
+	}
+
+	/**
+	 * Pop a Metadata sample from the cache.
+	 *
+	 * @return True if the operation succeeds.
+	 * @see AddMetadata, NumMetadata
+	 */
+	bool PopMetadata()
+	{
+		return MetadataSampleQueue.Pop();
+	}
+
+	/**
+	 * Pop a Subtitle sample from the cache.
+	 *
+	 * @return True if the operation succeeds.
+	 * @see AddSubtitle, NumSubtitle
+	 */
+	bool PopSubtitle()
+	{
+		return SubtitleSampleQueue.Pop();
+	}
+
+	/**
+	 * Pop a video sample from the cache.
+	 *
+	 * @return True if the operation succeeds.
+	 * @see AddVideo, NumVideo
+	 */
+	bool PopVideo()
+	{
+		return VideoSampleQueue.Pop();
+	}
+
+	/**
+	 * Get the number of queued audio samples.
+	 *
+	 * @return Number of samples.
+	 * @see AddAudio, PopAudio
+	 */
+	int32 NumAudioSamples() const
+	{
+		return AudioSampleQueue.Num();
+	}
+
+	/**
+	 * Get the number of queued caption samples.
+	 *
+	 * @return Number of samples.
+	 * @see AddCaption, PopCaption
+	 */
+	int32 NumCaptionSamples() const
+	{
+		return CaptionSampleQueue.Num();
+	}
+
+	/**
+	 * Get the number of queued metadata samples.
+	 *
+	 * @return Number of samples.
+	 * @see AddMetadata, PopMetada
+	 */
+	int32 NumMetadataSamples() const
+	{
+		return MetadataSampleQueue.Num();
+	}
+
+	/**
+	 * Get the number of queued subtitle samples.
+	 *
+	 * @return Number of samples.
+	 * @see AddSubtitle, PopSubtitle
+	 */
+	int32 NumSubtitleSamples() const
+	{
+		return SubtitleSampleQueue.Num();
+	}
+
+	/**
+	 * Get the number of queued video samples.
+	 *
+	 * @return Number of samples.
+	 * @see AddVideo, PopVideo
+	 */
+	int32 NumVideoSamples() const
+	{
+		return VideoSampleQueue.Num();
+	}
+
+	/**
+	 * Get next sample time from the VideoSampleQueue.
+	 *
+	 * @return Time of the next sample from the VideoSampleQueue
+	 * @see AddVideo, NumVideoSamples
+	 */
+	FTimespan GetNextVideoSampleTime()
+	{
+		FTimespan NextTime;
+		TSharedPtr<IMediaTextureSample, ESPMode::ThreadSafe> Sample;
+		const bool bHasSucceed = VideoSampleQueue.Peek(Sample);
+		if (bHasSucceed)
+		{
+			NextTime = Sample->GetTime();
+		}
+		return NextTime;
+	}
+
+public:
+
+	//~ IMediaSamples interface
+
+	virtual bool FetchAudio(TRange<FTimespan> TimeRange, TSharedPtr<IMediaAudioSample, ESPMode::ThreadSafe>& OutSample) override;
+	virtual bool FetchCaption(TRange<FTimespan> TimeRange, TSharedPtr<IMediaOverlaySample, ESPMode::ThreadSafe>& OutSample) override;
+	virtual bool FetchMetadata(TRange<FTimespan> TimeRange, TSharedPtr<IMediaBinarySample, ESPMode::ThreadSafe>& OutSample) override;
+	virtual bool FetchSubtitle(TRange<FTimespan> TimeRange, TSharedPtr<IMediaOverlaySample, ESPMode::ThreadSafe>& OutSample) override;
+	virtual bool FetchVideo(TRange<FTimespan> TimeRange, TSharedPtr<IMediaTextureSample, ESPMode::ThreadSafe>& OutSample) override;
+	virtual void FlushSamples() override;
+
+protected:
+
+	/** Audio sample queue. */
+	FMediaAudioSampleQueue AudioSampleQueue;
+
+	/** Caption sample queue. */
+	FMediaOverlaySampleQueue CaptionSampleQueue;
+
+	/** Metadata sample queue. */
+	FMediaBinarySampleQueue MetadataSampleQueue;
+
+	/** Subtitle sample queue. */
+	FMediaOverlaySampleQueue SubtitleSampleQueue;
+
+	/** Video sample queue. */
+	FMediaTextureSampleQueue VideoSampleQueue;
+};
diff --git a/Engine/Plugins/Media/MediaIOFramework/Source/MediaIOCore/Public/MediaIOCoreTextureSampleBase.h b/Engine/Plugins/Media/MediaIOFramework/Source/MediaIOCore/Public/MediaIOCoreTextureSampleBase.h
new file mode 100644
--- /dev/null
+++ b/Engine/Plugins/Media/MediaIOFramework/Source/MediaIOCore/Public/MediaIOCoreTextureSampleBase.h
@@ -0,0 +1,155 @@
+// Copyright 1998-2018 Epic Games, Inc. All Rights Reserved.
+
+#pragma once
+
+#include "IMediaTextureSample.h"
+#include "MediaObjectPool.h"
+
+/**
+ * Implements the IMediaTextureSample/IMediaPoolable interface.
+ */
+class MEDIAIOCORE_API FMediaIOCoreTextureSampleBase
+	: public IMediaTextureSample
+	, public IMediaPoolable
+{
+
+public:
+	FMediaIOCoreTextureSampleBase();
+
+	/**
+	 * Initialize the sample.
+	 *
+	 * @param InVideoBuffer The audio frame data.
+	 * @param InBufferSize The size of the video buffer.
+	 * @param InStride The number of channel of the video buffer.
+	 * @param InWidth The sample rate of the video buffer.
+	 * @param InHeight The sample rate of the video buffer.
+	 * @param InSampleFormat The sample format of the video buffer.
+	 * @param InTime The sample time (in the player's own clock).
+	 * @param InTimecode The sample timecode if available.
+	 */
+	bool Initialize(const void* InVideoBuffer, uint32 InBufferSize, uint32 InStride, uint32 InWidth, uint32 InHeight, EMediaTextureSampleFormat InSampleFormat, FTimespan InTime, const TOptional<FTimecode>& InTimecode);
+
+	/**
+	 * Initialize the sample.
+	 *
+	 * @param InVideoBuffer The audio frame data.
+	 * @param InStride The number of channel of the video buffer.
+	 * @param InWidth The sample rate of the video buffer.
+	 * @param InHeight The sample rate of the video buffer.
+	 * @param InSampleFormat The sample format of the video buffer.
+	 * @param InTime The sample time (in the player's own clock).
+	 * @param InTimecode The sample timecode if available.
+	 */
+	bool Initialize(TArray<uint8> InVideoBuffer, uint32 InStride, uint32 InWidth, uint32 InHeight, EMediaTextureSampleFormat InSampleFormat, FTimespan InTime, const TOptional<FTimecode>& InTimecode);
+
+	/**
+	 * Initialize the sample with half it's original height and take only the odd or even line.
+	 *
+	 * @param bUseEvenLine Should use the Even or the Odd line from the video buffer.
+	 * @param InVideoBuffer The audio frame data.
+	 * @param InStride The number of channel of the video buffer.
+	 * @param InWidth The sample rate of the video buffer.
+	 * @param InHeight The sample rate of the video buffer.
+	 * @param InSampleFormat The sample format of the video buffer.
+	 * @param InTime The sample time (in the player's own clock).
+	 * @param InTimecode The sample timecode if available.
+	 */
+	bool InitializeWithEvenOddLine(bool bUseEvenLine, const void* InVideoBuffer, uint32 InBufferSize, uint32 InStride, uint32 InWidth, uint32 InHeight, EMediaTextureSampleFormat InSampleFormat, FTimespan InTime, const TOptional<FTimecode>& InTimecode);
+
+public:
+	//~ IMediaTextureSample interface
+
+	virtual const void* GetBuffer() override
+	{
+		return Buffer.GetData();
+	}
+
+	virtual FIntPoint GetDim() const override
+	{
+		return FIntPoint(Width, Height);
+	}
+
+	virtual FTimespan GetDuration() const override
+	{
+		return FTimespan(ETimespan::TicksPerSecond / 60);
+	}
+
+	virtual EMediaTextureSampleFormat GetFormat() const override
+	{
+		return SampleFormat;
+	}
+
+	virtual FIntPoint GetOutputDim() const override
+	{
+		return FIntPoint(Width, Height);
+	}
+
+	virtual uint32 GetStride() const override
+	{
+		return Stride;
+	}
+
+#if WITH_ENGINE
+	virtual FRHITexture* GetTexture() const override
+	{
+		return nullptr;
+	}
+#endif //WITH_ENGINE
+
+	virtual FTimespan GetTime() const override
+	{
+		return Time;
+	}
+
+	virtual TOptional<FTimecode> GetTimecode() const override
+	{
+		return Timecode;
+	}
+
+	virtual bool IsCacheable() const override
+	{
+		return true;
+	}
+
+	virtual bool IsOutputSrgb() const override
+	{
+		return true;
+	}
+
+public:
+	//~ IMediaPoolable interface
+
+	virtual void ShutdownPoolable() override
+	{
+		FreeSample();
+	}
+
+protected:
+	virtual void FreeSample()
+	{
+		Buffer.Reset();
+	}
+
+protected:
+	/** Duration for which the sample is valid. */
+	FTimespan Duration;
+
+	/** Sample format. */
+	EMediaTextureSampleFormat SampleFormat;
+
+	/** Sample time. */
+	FTimespan Time;
+
+	/** Sample timecode. */
+	TOptional<FTimecode> Timecode;
+
+	/** Image dimensions */
+	uint32 Stride;
+	uint32 Width;
+	uint32 Height;
+
+	/** Pointer to raw pixels */
+	TArray<uint8> Buffer;
+};
+
diff --git a/Engine/Plugins/Media/MediaIOFramework/Source/MediaIOCore/Public/MediaIOCoreWaitVsyncThread.h b/Engine/Plugins/Media/MediaIOFramework/Source/MediaIOCore/Public/MediaIOCoreWaitVsyncThread.h
new file mode 100644
--- /dev/null
+++ b/Engine/Plugins/Media/MediaIOFramework/Source/MediaIOCore/Public/MediaIOCoreWaitVsyncThread.h
@@ -0,0 +1,34 @@
+// Copyright 1998-2018 Epic Games, Inc. All Rights Reserved.
+
+#pragma once
+
+#include "CoreMinimal.h"
+#include "Engine/EngineBaseTypes.h"
+#include "HAL/Runnable.h"
+
+/**
+* Implementation of WaitSyncThread
+*/
+
+class IMediaIOCoreHardwareSync;
+
+class MEDIAIOCORE_API FMediaIOCoreWaitVSyncThread : public FRunnable
+{
+public:
+	FMediaIOCoreWaitVSyncThread(TSharedPtr<IMediaIOCoreHardwareSync> InHardwareSync);
+
+	virtual bool Init() override;
+	virtual uint32 Run() override;
+
+	virtual void Stop() override;
+	virtual void Exit() override;
+
+public:
+	bool Wait_GameOrRenderThread();
+
+protected:
+	TSharedPtr<IMediaIOCoreHardwareSync> HardwareSync;
+	FEvent* WaitVSync;
+	TAtomic<bool> bWaitingForSignal;
+	TAtomic<bool> bAlive;
+};
diff --git a/Engine/Plugins/Media/MediaIOFramework/Source/MediaIOCore/Public/MediaOutput.h b/Engine/Plugins/Media/MediaIOFramework/Source/MediaIOCore/Public/MediaOutput.h
new file mode 100644
--- /dev/null
+++ b/Engine/Plugins/Media/MediaIOFramework/Source/MediaIOCore/Public/MediaOutput.h
@@ -0,0 +1,54 @@
+// Copyright 1998-2018 Epic Games, Inc. All Rights Reserved.
+
+#pragma once
+
+#include "CoreMinimal.h"
+#include "UObject/Object.h"
+#include "UObject/ObjectMacros.h"
+
+#include "PixelFormat.h"
+
+#include "MediaOutput.generated.h"
+
+
+class UMediaCapture;
+
+/**
+ * Abstract base class for media output.
+ *
+ * Media output describe the location and/or settings of media objects that can
+ * be used to output the content of UE4 to a target device via a MediaCapture.
+ */
+UCLASS(Abstract, editinlinenew, BlueprintType, hidecategories = (Object))
+class MEDIAIOCORE_API UMediaOutput : public UObject
+{
+	GENERATED_UCLASS_BODY()
+
+public:
+	/**
+	 * Number of texture used to transfer the texture from the GPU to the system memory.
+	 * A smaller number is most likely to block the GPU (wait for the transfer to complete).
+	 * A bigger number is most likely to increase latency.
+	 */
+	UPROPERTY(EditAnywhere, AdvancedDisplay, Category="Output", meta=(ClampMin=1, ClampMax=4))
+	int32 NumberOfTextureBuffers;
+
+	/** Creates the specific implementation of the MediaCapture for the MediaOutput. */
+	UFUNCTION(BlueprintCallable, Category="Media|Output")
+	UMediaCapture* CreateMediaCapture();
+
+	/**
+	 * Validate the media output settings (must be implemented in child classes).
+	 *
+	 * @return true if validation passed, false otherwise.
+	 */
+	UFUNCTION(BlueprintCallable, Category = "Media|Output")
+	virtual bool Validate(FString& OutFailureReason) const;
+
+public:
+	virtual FIntPoint GetRequestedSize() const PURE_VIRTUAL(UMediaOutput::GetRequestedSize, return FIntPoint(0, 0); );
+	virtual EPixelFormat GetRequestedPixelFormat() const PURE_VIRTUAL(UMediaOutput::GetRequestedPixelFormat, return EPixelFormat::PF_Unknown; );
+
+protected:
+	virtual UMediaCapture* CreateMediaCaptureImpl() { return nullptr; }
+};
diff --git a/Engine/Plugins/Media/MediaIOFramework/Source/MediaIOEditor/MediaIOEditor.Build.cs b/Engine/Plugins/Media/MediaIOFramework/Source/MediaIOEditor/MediaIOEditor.Build.cs
new file mode 100644
--- /dev/null
+++ b/Engine/Plugins/Media/MediaIOFramework/Source/MediaIOEditor/MediaIOEditor.Build.cs
@@ -0,0 +1,22 @@
+// Copyright 1998-2018 Epic Games, Inc. All Rights Reserved.
+
+namespace UnrealBuildTool.Rules
+{
+	public class MediaIOEditor : ModuleRules
+	{
+		public MediaIOEditor(ReadOnlyTargetRules Target) : base(Target)
+		{
+			PublicDependencyModuleNames.AddRange(
+				new string[]
+				{
+					"Core",
+					"CoreUObject",
+					"MediaIOCore",
+					"Slate",
+					"SlateCore",
+					"TimeManagement",
+					"UnrealEd",
+				});
+		}
+	}
+}
diff --git a/Engine/Plugins/Media/MediaIOFramework/Source/MediaIOEditor/Private/Customizations/MediaIOConfigurationCustomization.cpp b/Engine/Plugins/Media/MediaIOFramework/Source/MediaIOEditor/Private/Customizations/MediaIOConfigurationCustomization.cpp
new file mode 100644
--- /dev/null
+++ b/Engine/Plugins/Media/MediaIOFramework/Source/MediaIOEditor/Private/Customizations/MediaIOConfigurationCustomization.cpp
@@ -0,0 +1,129 @@
+// Copyright 1998-2018 Epic Games, Inc. All Rights Reserved.
+
+#include "Customizations/MediaIOConfigurationCustomization.h"
+
+#include "DetailWidgetRow.h"
+#include "DetailCategoryBuilder.h"
+#include "DetailLayoutBuilder.h"
+#include "Framework/Application/SlateApplication.h"
+#include "IDetailChildrenBuilder.h"
+#include "IMediaIOCoreDeviceProvider.h"
+#include "IMediaIOCoreModule.h"
+#include "MediaIOPermutationsSelectorBuilder.h"
+#include "ObjectEditorUtils.h"
+#include "Widgets/Input/SComboButton.h"
+#include "Widgets/Text/STextBlock.h"
+#include "Widgets/SBoxPanel.h"
+#include "Widgets/SMediaPermutationsSelector.h"
+#include "Widgets/SNullWidget.h"
+#include "Widgets/SWindow.h"
+
+#define LOCTEXT_NAMESPACE "MediaIOConfigurationCustomization"
+
+TSharedRef<IPropertyTypeCustomization> FMediaIOConfigurationCustomization::MakeInstance()
+{
+	return MakeShareable(new FMediaIOConfigurationCustomization);
+}
+
+TAttribute<FText> FMediaIOConfigurationCustomization::GetContentText()
+{
+	FMediaIOConfiguration* Value = GetPropertyValueFromPropertyHandle<FMediaIOConfiguration>();
+	IMediaIOCoreDeviceProvider* DeviceProviderPtr = IMediaIOCoreModule::Get().GetDeviceProvider(DeviceProviderName);
+	if (DeviceProviderPtr)
+	{
+		return MakeAttributeLambda([=] { return DeviceProviderPtr->ToText(*Value); });
+	}
+	return FText::GetEmpty();
+}
+
+TSharedRef<SWidget> FMediaIOConfigurationCustomization::HandleSourceComboButtonMenuContent()
+{
+	PermutationSelector.Reset();
+
+	IMediaIOCoreDeviceProvider* DeviceProviderPtr = IMediaIOCoreModule::Get().GetDeviceProvider(DeviceProviderName);
+	if (DeviceProviderPtr == nullptr)
+	{
+		return SNew(STextBlock)
+			.Text(LOCTEXT("NoDeviceProviderFound", "No provider found"));
+	}
+
+	SelectedConfiguration = *GetPropertyValueFromPropertyHandle<FMediaIOConfiguration>();
+	bool bIsInput = SelectedConfiguration.bIsInput;
+	if (!SelectedConfiguration.IsValid())
+	{
+		SelectedConfiguration = DeviceProviderPtr->GetDefaultConfiguration();
+		SelectedConfiguration.bIsInput = bIsInput;
+	}
+
+
+	TArray<FMediaIOConfiguration> MediaConfigurations = bIsInput ? DeviceProviderPtr->GetConfigurations(true, false) : DeviceProviderPtr->GetConfigurations(false, true);
+	if (MediaConfigurations.Num() == 0)
+	{
+		return SNew(STextBlock)
+			.Text(LOCTEXT("NoConfigurationFound", "No configuration found"));
+	}
+
+	auto QuadTypeVisible = [](FName ColumnName, const TArray<FMediaIOConfiguration>& UniquePermutationsForThisColumn)
+	{
+		if (UniquePermutationsForThisColumn.Num() > 0)
+		{
+			return UniquePermutationsForThisColumn[0].MediaConnection.TransportType == EMediaIOSDITransportType::QuadLink;
+		}
+		return false;
+	};
+
+	using TSelection = SMediaPermutationsSelector<FMediaIOConfiguration, FMediaIOPermutationsSelectorBuilder>;
+	TSelection::FArguments Arguments;
+	Arguments
+		.PermutationsSource(MoveTemp(MediaConfigurations))
+		.SelectedPermutation(SelectedConfiguration)
+		.OnSelectionChanged(this, &FMediaIOConfigurationCustomization::OnSelectionChanged)
+		.OnButtonClicked(this, &FMediaIOConfigurationCustomization::OnButtonClicked)
+		+ TSelection::Column(FMediaIOPermutationsSelectorBuilder::NAME_DeviceIdentifier)
+		.Label(LOCTEXT("DeviceLabel", "Device"));
+
+	if (DeviceProviderPtr->ShowInputTransportInSelector())
+	{
+		Arguments
+			+ TSelection::Column(FMediaIOPermutationsSelectorBuilder::NAME_TransportType)
+			.Label(LOCTEXT("SourceTypeLabel", "Source"))
+			+TSelection::Column(FMediaIOPermutationsSelectorBuilder::NAME_QuadType)
+			.Label(LOCTEXT("QuadTypeLabel", "Quad"))
+			.IsColumnVisible_Lambda(QuadTypeVisible);
+	}
+
+	Arguments
+		+ TSelection::Column(FMediaIOPermutationsSelectorBuilder::NAME_Resolution)
+		.Label(LOCTEXT("ResolutionLabel", "Resolution"))
+		+ TSelection::Column(FMediaIOPermutationsSelectorBuilder::NAME_Standard)
+		.Label(LOCTEXT("StandardLabel", "Standard"))
+		+ TSelection::Column(FMediaIOPermutationsSelectorBuilder::NAME_FrameRate)
+		.Label(LOCTEXT("FrameRateLabel", "Frame Rate"));
+
+	TSharedRef<TSelection> Selector = SNew(TSelection) = Arguments;
+	PermutationSelector = Selector;
+	SelectedConfiguration = Selector->GetSelectedItem();
+
+	return Selector;
+}
+
+void FMediaIOConfigurationCustomization::OnSelectionChanged(FMediaIOConfiguration SelectedItem)
+{
+	SelectedConfiguration = SelectedItem;
+}
+
+FReply FMediaIOConfigurationCustomization::OnButtonClicked() const
+{
+	AssignValue(SelectedConfiguration);
+
+	TSharedPtr<SWidget> SharedPermutationSelector = PermutationSelector.Pin();
+	if (SharedPermutationSelector.IsValid())
+	{
+		TSharedRef<SWindow> ParentContextMenuWindow = FSlateApplication::Get().FindWidgetWindow(SharedPermutationSelector.ToSharedRef()).ToSharedRef();
+		FSlateApplication::Get().RequestDestroyWindow(ParentContextMenuWindow);
+	}
+
+	return FReply::Handled();
+}
+
+#undef LOCTEXT_NAMESPACE
diff --git a/Engine/Plugins/Media/MediaIOFramework/Source/MediaIOEditor/Private/Customizations/MediaIOCustomizationBase.cpp b/Engine/Plugins/Media/MediaIOFramework/Source/MediaIOEditor/Private/Customizations/MediaIOCustomizationBase.cpp
new file mode 100644
--- /dev/null
+++ b/Engine/Plugins/Media/MediaIOFramework/Source/MediaIOEditor/Private/Customizations/MediaIOCustomizationBase.cpp
@@ -0,0 +1,133 @@
+// Copyright 1998-2018 Epic Games, Inc. All Rights Reserved.
+
+#include "Customizations/MediaIOCustomizationBase.h"
+
+#include "MediaIOCoreDefinitions.h"
+#include "DetailWidgetRow.h"
+#include "Framework/MultiBox/MultiBoxBuilder.h"
+#include "IDetailChildrenBuilder.h"
+#include "IDetailPropertyRow.h"
+#include "IPropertyUtilities.h"
+#include "Widgets/Input/SComboButton.h"
+
+#define LOCTEXT_NAMESPACE "MediaIOCustomizationBase"
+
+namespace FMediaIOCustomizationBaseHelper
+{
+	FName TraverseClassHierarchy(UClass* InClass, const FName& InMetadataName)
+	{
+		if (InClass == nullptr)
+		{
+			return NAME_None;
+		}
+
+		if (!InClass->HasMetaData(InMetadataName))
+		{
+			return TraverseClassHierarchy(InClass->GetSuperClass(), InMetadataName);
+		}
+
+		return *InClass->GetMetaData(InMetadataName);
+	}
+}
+
+void FMediaIOCustomizationBase::CustomizeHeader(TSharedRef<IPropertyHandle> InPropertyHandle, FDetailWidgetRow& HeaderRow, IPropertyTypeCustomizationUtils& CustomizationUtils)
+{
+	MediaProperty = InPropertyHandle;
+
+	// Read the media provider name
+	if (MediaProperty->IsValidHandle())
+	{
+		static FName NAME_MediaIOCustomLayout = TEXT("MediaIOCustomLayout");
+		if (MediaProperty->HasMetaData(NAME_MediaIOCustomLayout))
+		{
+			DeviceProviderName = *MediaProperty->GetPropertyClass()->GetMetaData(NAME_MediaIOCustomLayout);
+		}
+
+		if (DeviceProviderName == NAME_None)
+		{
+			TArray<UObject*> OuterObjects;
+			MediaProperty->GetOuterObjects(OuterObjects);
+
+			for (UObject* Obj : OuterObjects)
+			{
+				FName FoundMetadata = FMediaIOCustomizationBaseHelper::TraverseClassHierarchy(Obj->GetClass(), NAME_MediaIOCustomLayout);
+				if (FoundMetadata == NAME_None)
+				{
+					DeviceProviderName = NAME_None;
+					break;
+				}
+
+				if (DeviceProviderName == NAME_None)
+				{
+					DeviceProviderName = FoundMetadata;
+				}
+				else if (DeviceProviderName != FoundMetadata)
+				{
+					DeviceProviderName = NAME_None;
+					break;
+				}
+			}
+		}
+		ensure(DeviceProviderName != NAME_None);
+	}
+
+	if (MediaProperty->GetNumPerObjectValues() == 1 && MediaProperty->IsValidHandle())
+	{
+		TSharedPtr<IPropertyUtilities> PropertyUtils = CustomizationUtils.GetPropertyUtilities();
+
+		HeaderRow
+		.NameContent()
+		[
+			InPropertyHandle->CreatePropertyNameWidget()
+		]
+		.ValueContent()
+		.MaxDesiredWidth(512)
+		[
+			SNew(SHorizontalBox)
+			+ SHorizontalBox::Slot()
+			.FillWidth(1.0f)
+			.VAlign(VAlign_Center)
+			[
+				SNew(STextBlock)
+				.Text(GetContentText())
+			]
+			+ SHorizontalBox::Slot()
+			.AutoWidth()
+			.Padding(FMargin(4.0f, 0.0f, 0.0f, 0.0f))
+			.VAlign(VAlign_Center)
+			[
+				SNew(SComboButton)
+				.OnGetMenuContent_Lambda( [this]() { return HandleSourceComboButtonMenuContent(); } )
+				.ContentPadding(FMargin(4.0, 2.0))
+			]
+		].IsEnabled(MakeAttributeLambda([=] { return !InPropertyHandle->IsEditConst() && PropertyUtils->IsPropertyEditingEnabled(); }));
+	}
+}
+
+void FMediaIOCustomizationBase::CustomizeChildren(TSharedRef<IPropertyHandle> InStructPropertyHandle, IDetailChildrenBuilder& StructBuilder, IPropertyTypeCustomizationUtils& StructCustomizationUtils)
+{
+	uint32 NumberOfChild;
+	if (InStructPropertyHandle->GetNumChildren(NumberOfChild) == FPropertyAccess::Success)
+	{
+		for (uint32 Index = 0; Index < NumberOfChild; ++Index)
+		{
+			StructBuilder.AddProperty(InStructPropertyHandle->GetChildHandle(Index).ToSharedRef()).IsEnabled(false).ShowPropertyButtons(false);
+		}
+	}
+	InStructPropertyHandle->MarkHiddenByCustomization();
+}
+
+void FMediaIOCustomizationBase::AssignValueImpl(const void* NewValue) const
+{
+	if (UStructProperty* StructProperty = Cast<UStructProperty>(MediaProperty->GetProperty()))
+	{
+		TArray<void*> RawData;
+		MediaProperty->AccessRawData(RawData);
+
+		FString TextValue;
+		StructProperty->Struct->ExportText(TextValue, NewValue, RawData[0], nullptr, EPropertyPortFlags::PPF_None, nullptr);
+		ensure(MediaProperty->SetValueFromFormattedString(TextValue, EPropertyValueSetFlags::DefaultFlags) == FPropertyAccess::Result::Success);
+	}
+}
+
+#undef LOCTEXT_NAMESPACE
diff --git a/Engine/Plugins/Media/MediaIOFramework/Source/MediaIOEditor/Private/Customizations/MediaIODeviceCustomization.cpp b/Engine/Plugins/Media/MediaIOFramework/Source/MediaIOEditor/Private/Customizations/MediaIODeviceCustomization.cpp
new file mode 100644
--- /dev/null
+++ b/Engine/Plugins/Media/MediaIOFramework/Source/MediaIOEditor/Private/Customizations/MediaIODeviceCustomization.cpp
@@ -0,0 +1,83 @@
+// Copyright 1998-2018 Epic Games, Inc. All Rights Reserved.
+
+#include "Customizations/MediaIODeviceCustomization.h"
+
+#include "DetailWidgetRow.h"
+#include "Framework/MultiBox/MultiBoxBuilder.h"
+#include "IMediaIOCoreDeviceProvider.h"
+#include "IMediaIOCoreModule.h"
+#include "IPropertyUtilities.h"
+#include "MediaIOCoreDefinitions.h"
+#include "PropertyHandle.h"
+#include "Widgets/Input/SComboButton.h"
+
+#define LOCTEXT_NAMESPACE "MediaIODeviceCustomization"
+
+TSharedRef<IPropertyTypeCustomization> FMediaIODeviceCustomization::MakeInstance()
+{
+	TSharedRef<FMediaIODeviceCustomization> Result = MakeShareable(new FMediaIODeviceCustomization);
+	return Result;
+}
+
+TAttribute<FText> FMediaIODeviceCustomization::GetContentText()
+{
+	FMediaIODevice* Value = GetPropertyValueFromPropertyHandle<FMediaIODevice>();
+	IMediaIOCoreDeviceProvider* DeviceProviderPtr = IMediaIOCoreModule::Get().GetDeviceProvider(DeviceProviderName);
+	if (DeviceProviderPtr)
+	{
+		return MakeAttributeLambda([=] { return DeviceProviderPtr->ToText(*Value); });
+	}
+	return FText::GetEmpty();
+}
+
+TSharedRef<SWidget> FMediaIODeviceCustomization::HandleSourceComboButtonMenuContent()
+{
+	IMediaIOCoreDeviceProvider* DeviceProviderPtr = IMediaIOCoreModule::Get().GetDeviceProvider(DeviceProviderName);
+	if (DeviceProviderPtr == nullptr)
+	{
+		return SNew(STextBlock)
+			.Text(LOCTEXT("NoDeviceProviderFound", "No provider found"));
+	}
+
+	// found all devices
+	TArray<FMediaIODevice> AllDevices = DeviceProviderPtr->GetDevices();
+	if (AllDevices.Num() == 0)
+	{
+		return SNew(STextBlock)
+			.Text(LOCTEXT("NoDeviceFound", "No device found"));
+	}
+
+	// generate menu
+	FMenuBuilder MenuBuilder(true, nullptr);
+
+	MenuBuilder.BeginSection("AllDevices", LOCTEXT("AllDevicesSection", "All Devices"));
+	{
+		for (const FMediaIODevice& Device : AllDevices)
+		{
+			MenuBuilder.AddMenuEntry(
+				DeviceProviderPtr->ToText(Device),
+				DeviceProviderPtr->ToText(Device),
+				FSlateIcon(),
+				FUIAction(
+					FExecuteAction::CreateLambda([=] {
+						AssignValue(Device);
+					}),
+					FCanExecuteAction(),
+					FIsActionChecked::CreateLambda([this, Device] {
+						TArray<void*> RawData;
+						GetMediaProperty()->AccessRawData(RawData);
+						FMediaIODevice* PreviousDevice = reinterpret_cast<FMediaIODevice*>(RawData[0]);
+						return *PreviousDevice == Device;
+					})
+				),
+				NAME_None,
+				EUserInterfaceActionType::RadioButton
+				);
+		}
+	}
+	MenuBuilder.EndSection();
+
+	return MenuBuilder.MakeWidget();
+}
+
+#undef LOCTEXT_NAMESPACE
diff --git a/Engine/Plugins/Media/MediaIOFramework/Source/MediaIOEditor/Private/Customizations/MediaIOOutputConfigurationCustomization.cpp b/Engine/Plugins/Media/MediaIOFramework/Source/MediaIOEditor/Private/Customizations/MediaIOOutputConfigurationCustomization.cpp
new file mode 100644
--- /dev/null
+++ b/Engine/Plugins/Media/MediaIOFramework/Source/MediaIOEditor/Private/Customizations/MediaIOOutputConfigurationCustomization.cpp
@@ -0,0 +1,161 @@
+// Copyright 1998-2018 Epic Games, Inc. All Rights Reserved.
+
+#include "Customizations/MediaIOOutputConfigurationCustomization.h"
+
+#include "DetailWidgetRow.h"
+#include "DetailCategoryBuilder.h"
+#include "DetailLayoutBuilder.h"
+#include "Framework/Application/SlateApplication.h"
+#include "IDetailChildrenBuilder.h"
+#include "IMediaIOCoreDeviceProvider.h"
+#include "IMediaIOCoreModule.h"
+#include "MediaIOPermutationsSelectorBuilder.h"
+#include "ObjectEditorUtils.h"
+#include "Widgets/Input/SComboButton.h"
+#include "Widgets/Text/STextBlock.h"
+#include "Widgets/SBoxPanel.h"
+#include "Widgets/SMediaPermutationsSelector.h"
+#include "Widgets/SNullWidget.h"
+#include "Widgets/SWindow.h"
+
+#define LOCTEXT_NAMESPACE "MediaIOOutputConfigurationCustomization"
+
+TSharedRef<IPropertyTypeCustomization> FMediaIOOutputConfigurationCustomization::MakeInstance()
+{
+	return MakeShareable(new FMediaIOOutputConfigurationCustomization);
+}
+
+TAttribute<FText> FMediaIOOutputConfigurationCustomization::GetContentText()
+{
+	FMediaIOOutputConfiguration* Value = GetPropertyValueFromPropertyHandle<FMediaIOOutputConfiguration>();
+	IMediaIOCoreDeviceProvider* DeviceProviderPtr = IMediaIOCoreModule::Get().GetDeviceProvider(DeviceProviderName);
+	if (DeviceProviderPtr)
+	{
+		return MakeAttributeLambda([=] { return DeviceProviderPtr->ToText(*Value); });
+	}
+	return FText::GetEmpty();
+}
+
+TSharedRef<SWidget> FMediaIOOutputConfigurationCustomization::HandleSourceComboButtonMenuContent()
+{
+	PermutationSelector.Reset();
+
+	IMediaIOCoreDeviceProvider* DeviceProviderPtr = IMediaIOCoreModule::Get().GetDeviceProvider(DeviceProviderName);
+	if (DeviceProviderPtr == nullptr)
+	{
+		return SNew(STextBlock)
+			.Text(LOCTEXT("NoDeviceProviderFound", "No provider found"));
+	}
+
+	SelectedConfiguration = *GetPropertyValueFromPropertyHandle<FMediaIOOutputConfiguration>();
+	if (!SelectedConfiguration.IsValid())
+	{
+		SelectedConfiguration = DeviceProviderPtr->GetDefaultOutputConfiguration();
+	}
+
+
+	TArray<FMediaIOOutputConfiguration> MediaConfigurations = DeviceProviderPtr->GetOutputConfigurations();
+	if (MediaConfigurations.Num() == 0)
+	{
+		return SNew(STextBlock)
+			.Text(LOCTEXT("NoConfigurationFound", "No configuration found"));
+	}
+
+	// Create the permutation selector
+	auto QuadTypeVisible = [](FName ColumnName, const TArray<FMediaIOOutputConfiguration>& UniquePermutationsForThisColumn)
+	{
+		if (UniquePermutationsForThisColumn.Num() > 0)
+		{
+			return UniquePermutationsForThisColumn[0].MediaConfiguration.MediaConnection.TransportType == EMediaIOSDITransportType::QuadLink;
+		}
+		return false;
+	};
+
+	auto KeyVisible = [](FName ColumnName, const TArray<FMediaIOOutputConfiguration>& UniquePermutationsForThisColumn)
+	{
+		if (UniquePermutationsForThisColumn.Num() > 0)
+		{
+			return UniquePermutationsForThisColumn[0].OutputType == EMediaIOOutputType::FillAndKey;
+		}
+		return false;
+	};
+
+	auto SyncVisible = [](FName ColumnName, const TArray<FMediaIOOutputConfiguration>& UniquePermutationsForThisColumn)
+	{
+		if (UniquePermutationsForThisColumn.Num() > 0)
+		{
+			return UniquePermutationsForThisColumn[0].OutputReference == EMediaIOReferenceType::Input;
+		}
+		return false;
+	};
+
+	using TSelection = SMediaPermutationsSelector<FMediaIOOutputConfiguration, FMediaIOPermutationsSelectorBuilder>;
+	TSelection::FArguments Arguments;
+	Arguments
+		.PermutationsSource(MoveTemp(MediaConfigurations))
+		.SelectedPermutation(SelectedConfiguration)
+		.OnSelectionChanged(this, &FMediaIOOutputConfigurationCustomization::OnSelectionChanged)
+		.OnButtonClicked(this, &FMediaIOOutputConfigurationCustomization::OnButtonClicked)
+		+ TSelection::Column(FMediaIOPermutationsSelectorBuilder::NAME_OutputType)
+		.Label(LOCTEXT("OutputTypeeLabel", "Output Type"))
+		+ TSelection::Column(FMediaIOPermutationsSelectorBuilder::NAME_DeviceIdentifier)
+		.Label(LOCTEXT("DeviceLabel", "Device"));
+
+	if (DeviceProviderPtr->ShowOutputTransportInSelector())
+	{
+		Arguments
+			+ TSelection::Column(FMediaIOPermutationsSelectorBuilder::NAME_TransportType)
+			.Label(LOCTEXT("SourceTypeLabel", "Source"))
+			+ TSelection::Column(FMediaIOPermutationsSelectorBuilder::NAME_QuadType)
+			.Label(LOCTEXT("QuadTypeLabel", "Quad"))
+			.IsColumnVisible_Lambda(QuadTypeVisible);
+	}
+
+	Arguments
+		+ TSelection::Column(FMediaIOPermutationsSelectorBuilder::NAME_Resolution)
+		.Label(LOCTEXT("ResolutionLabel", "Resolution"))
+		+ TSelection::Column(FMediaIOPermutationsSelectorBuilder::NAME_Standard)
+		.Label(LOCTEXT("StandardLabel", "Standard"))
+		+ TSelection::Column(FMediaIOPermutationsSelectorBuilder::NAME_FrameRate)
+		.Label(LOCTEXT("FrameRateLabel", "Frame Rate"))
+		+ TSelection::Column(FMediaIOPermutationsSelectorBuilder::NAME_KeyPortSource)
+		.Label(LOCTEXT("KeySourceTypeLabel", "Key Source"))
+		.IsColumnVisible_Lambda(KeyVisible);
+
+	if (DeviceProviderPtr->ShowReferenceInSelector())
+	{
+		Arguments
+			+ TSelection::Column(FMediaIOPermutationsSelectorBuilder::NAME_OutputReference)
+			.Label(LOCTEXT("ReferenceLabel", "Reference"))
+			+ TSelection::Column(FMediaIOPermutationsSelectorBuilder::NAME_SyncPortSource)
+			.Label(LOCTEXT("SyncSourceTypeLabel", "Sync Source"))
+			.IsColumnVisible_Lambda(SyncVisible);
+	}
+
+	TSharedRef<TSelection> Selector = SNew(TSelection) = Arguments;
+	PermutationSelector = Selector;
+	SelectedConfiguration = Selector->GetSelectedItem();
+
+	return Selector;
+}
+
+void FMediaIOOutputConfigurationCustomization::OnSelectionChanged(FMediaIOOutputConfiguration SelectedItem)
+{
+	SelectedConfiguration = SelectedItem;
+}
+
+FReply FMediaIOOutputConfigurationCustomization::OnButtonClicked() const
+{
+	AssignValue(SelectedConfiguration);
+
+	TSharedPtr<SWidget> SharedPermutationSelector = PermutationSelector.Pin();
+	if (SharedPermutationSelector.IsValid())
+	{
+		TSharedRef<SWindow> ParentContextMenuWindow = FSlateApplication::Get().FindWidgetWindow(SharedPermutationSelector.ToSharedRef()).ToSharedRef();
+		FSlateApplication::Get().RequestDestroyWindow(ParentContextMenuWindow);
+	}
+
+	return FReply::Handled();
+}
+
+#undef LOCTEXT_NAMESPACE
diff --git a/Engine/Plugins/Media/MediaIOFramework/Source/MediaIOEditor/Private/MediaIOEditorModule.cpp b/Engine/Plugins/Media/MediaIOFramework/Source/MediaIOEditor/Private/MediaIOEditorModule.cpp
new file mode 100644
--- /dev/null
+++ b/Engine/Plugins/Media/MediaIOFramework/Source/MediaIOEditor/Private/MediaIOEditorModule.cpp
@@ -0,0 +1,53 @@
+// Copyright 1998-2018 Epic Games, Inc. All Rights Reserved.
+
+
+#include "MediaIOEditorModule.h"
+
+#include "Modules/ModuleInterface.h"
+#include "Modules/ModuleManager.h"
+#include "PropertyEditorDelegates.h"
+#include "PropertyEditorModule.h"
+
+#include "Customizations/MediaIODeviceCustomization.h"
+#include "Customizations/MediaIOConfigurationCustomization.h"
+#include "Customizations/MediaIOOutputConfigurationCustomization.h"
+
+
+DEFINE_LOG_CATEGORY(LogMediaIOEditor);
+
+class FMediaIOEditorModule : public IModuleInterface
+{
+	virtual void StartupModule() override
+	{
+		RegisterCustomizations();
+	}
+
+	virtual void ShutdownModule() override
+	{
+		if (!UObjectInitialized() && !GIsRequestingExit)
+		{
+			UnregisterCustomizations();
+		}
+	}
+
+private:
+	/** Register details view customizations. */
+	void RegisterCustomizations()
+	{
+		FPropertyEditorModule& PropertyModule = FModuleManager::LoadModuleChecked<FPropertyEditorModule>("PropertyEditor");
+		PropertyModule.RegisterCustomPropertyTypeLayout(FMediaIOConfiguration::StaticStruct()->GetFName(), FOnGetPropertyTypeCustomizationInstance::CreateStatic(&FMediaIOConfigurationCustomization::MakeInstance));
+		PropertyModule.RegisterCustomPropertyTypeLayout(FMediaIODevice::StaticStruct()->GetFName(), FOnGetPropertyTypeCustomizationInstance::CreateStatic(&FMediaIODeviceCustomization::MakeInstance));
+		PropertyModule.RegisterCustomPropertyTypeLayout(FMediaIOOutputConfiguration::StaticStruct()->GetFName(), FOnGetPropertyTypeCustomizationInstance::CreateStatic(&FMediaIOOutputConfigurationCustomization::MakeInstance));
+	}
+
+	/** Unregister details view customizations. */
+	void UnregisterCustomizations()
+	{
+		FPropertyEditorModule& PropertyModule = FModuleManager::LoadModuleChecked<FPropertyEditorModule>("PropertyEditor");
+		PropertyModule.UnregisterCustomPropertyTypeLayout(FMediaIODevice::StaticStruct()->GetFName());
+		PropertyModule.UnregisterCustomPropertyTypeLayout(FMediaIOConfiguration::StaticStruct()->GetFName());
+		PropertyModule.UnregisterCustomPropertyTypeLayout(FMediaIOOutputConfiguration::StaticStruct()->GetFName());
+	}
+};
+
+IMPLEMENT_MODULE(FMediaIOEditorModule, MediaIOEditor)
diff --git a/Engine/Plugins/Media/MediaIOFramework/Source/MediaIOEditor/Private/MediaIOEditorModule.h b/Engine/Plugins/Media/MediaIOFramework/Source/MediaIOEditor/Private/MediaIOEditorModule.h
new file mode 100644
--- /dev/null
+++ b/Engine/Plugins/Media/MediaIOFramework/Source/MediaIOEditor/Private/MediaIOEditorModule.h
@@ -0,0 +1,7 @@
+// Copyright 1998-2018 Epic Games, Inc. All Rights Reserved.
+#pragma once
+
+#include "CoreMinimal.h"
+
+
+DECLARE_LOG_CATEGORY_EXTERN(LogMediaIOEditor, Log, All);
diff --git a/Engine/Plugins/Media/MediaIOFramework/Source/MediaIOEditor/Private/MediaIOPermutationsSelectorBuilder.cpp b/Engine/Plugins/Media/MediaIOFramework/Source/MediaIOEditor/Private/MediaIOPermutationsSelectorBuilder.cpp
new file mode 100644
--- /dev/null
+++ b/Engine/Plugins/Media/MediaIOFramework/Source/MediaIOEditor/Private/MediaIOPermutationsSelectorBuilder.cpp
@@ -0,0 +1,253 @@
+// Copyright 1998-2018 Epic Games, Inc. All Rights Reserved.
+
+#include "MediaIOPermutationsSelectorBuilder.h"
+
+#include "CommonFrameRates.h"
+#include "MediaIOCoreCommonDisplayMode.h"
+
+#define LOCTEXT_NAMESPACE "MediaIOPermutationsSelectorBuilder"
+
+namespace MediaIOPermutationsSelectorBuilder
+{
+	FText LinkTypeToPrettyText(EMediaIOSDITransportType InLinkType, int32 InChannel)
+	{
+		switch (InLinkType)
+		{
+		case EMediaIOSDITransportType::SingleLink: return FText::Format(LOCTEXT("SingleLinkLabel", "Single Link {0}"), FText::AsNumber(InChannel));
+		case EMediaIOSDITransportType::DualLink: return FText::Format(LOCTEXT("DualLinkLabel", "Dual Link {0}-{1}"), FText::AsNumber(InChannel), FText::AsNumber(InChannel + 1));
+		case EMediaIOSDITransportType::QuadLink: return FText::Format(LOCTEXT("QuadLinkLabel", "Quad Link {0}-{1}"), FText::AsNumber(InChannel), FText::AsNumber(InChannel + 3));
+		}
+		return FText::GetEmpty();
+	}
+
+	FText QuadLinkTypeToPrettyText(EMediaIOQuadLinkTransportType InLinkType)
+	{
+		switch (InLinkType)
+		{
+		case EMediaIOQuadLinkTransportType::SquareDivision: return LOCTEXT("QuadLinkSquareLabel", "Square Division");
+		case EMediaIOQuadLinkTransportType::TwoSampleInterleave: return LOCTEXT("QuadLinkSILabel", "Sample Interleave");
+		}
+		return FText::GetEmpty();
+	}
+
+	FText StandardTypeToPrettyText(EMediaIOStandardType InStandard)
+	{
+		switch (InStandard)
+		{
+		case EMediaIOStandardType::Progressive: return LOCTEXT("ProgressiveLabel", "Progressive");
+		case EMediaIOStandardType::Interlaced: return LOCTEXT("InterlacedLabel", "Interlaced");
+		case EMediaIOStandardType::ProgressiveSegmentedFrame: return LOCTEXT("ProgressiveSegmentedFrameLabel", "PSF");
+		}
+		return FText::GetEmpty();
+	}
+}
+
+/**
+ * FNames
+ */
+const FName FMediaIOPermutationsSelectorBuilder::NAME_DeviceIdentifier = "DeviceIdentifier";
+const FName FMediaIOPermutationsSelectorBuilder::NAME_TransportType = "SourceType";
+const FName FMediaIOPermutationsSelectorBuilder::NAME_QuadType = "QuadType";
+const FName FMediaIOPermutationsSelectorBuilder::NAME_Resolution = "Resolution";
+const FName FMediaIOPermutationsSelectorBuilder::NAME_Standard = "Standard";
+const FName FMediaIOPermutationsSelectorBuilder::NAME_FrameRate = "FrameRate";
+
+const FName FMediaIOPermutationsSelectorBuilder::NAME_OutputType = "OutputType";
+const FName FMediaIOPermutationsSelectorBuilder::NAME_KeyPortSource = "KeyPortSource";
+const FName FMediaIOPermutationsSelectorBuilder::NAME_OutputReference = "OutputReference";
+const FName FMediaIOPermutationsSelectorBuilder::NAME_SyncPortSource = "SyncPortSource";
+
+/**
+ * With FMediaIOConnection
+ */
+bool FMediaIOPermutationsSelectorBuilder::IdenticalProperty(FName ColumnName, const FMediaIOConnection& Left, const FMediaIOConnection& Right)
+{
+	if (ColumnName == NAME_DeviceIdentifier) return Left.Device.DeviceIdentifier == Right.Device.DeviceIdentifier;
+	if (ColumnName == NAME_TransportType) return Left.TransportType == Right.TransportType && Left.PortIdentifier == Right.PortIdentifier;
+	if (ColumnName == NAME_QuadType) return Left.TransportType == EMediaIOSDITransportType::QuadLink ? Left.QuadTransportType == Right.QuadTransportType : true;
+
+	return false;
+}
+
+
+bool FMediaIOPermutationsSelectorBuilder::Less(FName ColumnName, const FMediaIOConnection& Left, const FMediaIOConnection& Right)
+{
+	if (ColumnName == NAME_DeviceIdentifier) return Left.Device.DeviceIdentifier < Right.Device.DeviceIdentifier;
+	if (ColumnName == NAME_TransportType)
+	{
+		if (Left.TransportType == Right.TransportType)
+		{
+			return Left.PortIdentifier < Right.PortIdentifier;
+		}
+		return (int32)Left.TransportType < (int32)Right.TransportType;
+	}
+	if (ColumnName == NAME_QuadType) return (Left.TransportType == EMediaIOSDITransportType::QuadLink) ? (int32)Left.QuadTransportType < (int32)Right.QuadTransportType : false;
+
+	return false;
+}
+
+
+FText FMediaIOPermutationsSelectorBuilder::GetLabel(FName ColumnName, const FMediaIOConnection& Item)
+{
+	if (ColumnName == NAME_DeviceIdentifier) return FText::FromName(Item.Device.DeviceName);
+	if (ColumnName == NAME_TransportType) return MediaIOPermutationsSelectorBuilder::LinkTypeToPrettyText(Item.TransportType, Item.PortIdentifier);
+	if (ColumnName == NAME_QuadType) return MediaIOPermutationsSelectorBuilder::QuadLinkTypeToPrettyText(Item.QuadTransportType);
+
+	return FText::GetEmpty();
+}
+
+
+FText FMediaIOPermutationsSelectorBuilder::GetTooltip(FName ColumnName, const FMediaIOConnection& Item)
+{
+	if (ColumnName == NAME_DeviceIdentifier) return FText::FromString(FString::Printf(TEXT("%s as identifier: %d"), *Item.Device.DeviceName.ToString(), Item.Device.DeviceIdentifier));
+	if (ColumnName == NAME_TransportType) return FText::GetEmpty();
+	if (ColumnName == NAME_QuadType) return LOCTEXT("QuadTypeTooltip", "Output can be Square Division Quad Split (SQ) or Two-Sample Interleave (SI).");
+
+	return FText::GetEmpty();
+}
+
+
+/**
+ * With FMediaIOMode
+ */
+bool FMediaIOPermutationsSelectorBuilder::IdenticalProperty(FName ColumnName, const FMediaIOMode& Left, const FMediaIOMode& Right)
+{
+	if (ColumnName == NAME_Resolution) return Left.Resolution == Right.Resolution;
+	if (ColumnName == NAME_Standard) return Left.Standard == Right.Standard;
+	if (ColumnName == NAME_FrameRate) return Left.FrameRate == Right.FrameRate;
+
+	return false;
+}
+
+
+bool FMediaIOPermutationsSelectorBuilder::Less(FName ColumnName, const FMediaIOMode& Left, const FMediaIOMode& Right)
+{
+	if (ColumnName == NAME_Resolution) return Left.Resolution.SizeSquared() < Right.Resolution.SizeSquared();
+	if (ColumnName == NAME_Standard) return (int32)Left.Standard < (int32)Right.Standard;
+	if (ColumnName == NAME_FrameRate) return Left.FrameRate.AsDecimal() < Right.FrameRate.AsDecimal();
+
+	return false;
+}
+
+
+FText FMediaIOPermutationsSelectorBuilder::GetLabel(FName ColumnName, const FMediaIOMode& Item)
+{
+	if (ColumnName == NAME_Resolution) return FMediaIOCommonDisplayModes::GetMediaIOCommonDisplayModeResolutionInfoName(Item.Resolution.X, Item.Resolution.Y);
+	if (ColumnName == NAME_Standard) return MediaIOPermutationsSelectorBuilder::StandardTypeToPrettyText(Item.Standard);
+	if (ColumnName == NAME_FrameRate) return Item.FrameRate.ToPrettyText();
+
+	return FText::GetEmpty();
+}
+
+
+FText FMediaIOPermutationsSelectorBuilder::GetTooltip(FName ColumnName, const FMediaIOMode& Item)
+{
+	if (ColumnName == NAME_Resolution) return FText::FromString(FString::Printf(TEXT("%dx%d"), Item.Resolution.X, Item.Resolution.Y));
+	if (ColumnName == NAME_Standard) return FText::GetEmpty();
+	if (ColumnName == NAME_FrameRate)
+	{
+		if (const FCommonFrameRateInfo* Found = FCommonFrameRates::Find(Item.FrameRate))
+		{
+			return Found->Description;
+		}
+		return Item.FrameRate.ToPrettyText();
+	}
+
+	return FText::GetEmpty();
+}
+
+
+/**
+ * With FMediaIOMode
+ */
+bool FMediaIOPermutationsSelectorBuilder::IdenticalProperty(FName ColumnName, const FMediaIOConfiguration& Left, const FMediaIOConfiguration& Right)
+{
+	if (ColumnName == NAME_Resolution || ColumnName == NAME_Standard || ColumnName == NAME_FrameRate)
+	{
+		return IdenticalProperty(ColumnName, Left.MediaMode, Right.MediaMode);
+	}
+	return IdenticalProperty(ColumnName, Left.MediaConnection, Right.MediaConnection);
+}
+
+
+bool FMediaIOPermutationsSelectorBuilder::Less(FName ColumnName, const FMediaIOConfiguration& Left, const FMediaIOConfiguration& Right)
+{
+	if (ColumnName == NAME_Resolution || ColumnName == NAME_Standard || ColumnName == NAME_FrameRate)
+	{
+		return Less(ColumnName, Left.MediaMode, Right.MediaMode);
+	}
+	return Less(ColumnName, Left.MediaConnection, Right.MediaConnection);
+}
+
+FText FMediaIOPermutationsSelectorBuilder::GetLabel(FName ColumnName, const FMediaIOConfiguration& Item)
+{
+	if (ColumnName == NAME_Resolution || ColumnName == NAME_Standard || ColumnName == NAME_FrameRate)
+	{
+		return GetLabel(ColumnName, Item.MediaMode);
+	}
+	return GetLabel(ColumnName, Item.MediaConnection);
+}
+
+
+FText FMediaIOPermutationsSelectorBuilder::GetTooltip(FName ColumnName, const FMediaIOConfiguration& Item)
+{
+	if (ColumnName == NAME_Resolution || ColumnName == NAME_Standard || ColumnName == NAME_FrameRate)
+	{
+		return GetTooltip(ColumnName, Item.MediaMode);
+	}
+	return GetTooltip(ColumnName, Item.MediaConnection);
+}
+
+
+bool FMediaIOPermutationsSelectorBuilder::IdenticalProperty(FName ColumnName, const FMediaIOOutputConfiguration& Left, const FMediaIOOutputConfiguration& Right)
+{
+	if (ColumnName == NAME_OutputType) return Left.OutputType == Right.OutputType;
+	if (ColumnName == NAME_KeyPortSource) return Left.OutputType == EMediaIOOutputType::FillAndKey ? Left.KeyPortIdentifier == Right.KeyPortIdentifier : true;
+	if (ColumnName == NAME_OutputReference) return Left.OutputReference == Right.OutputReference;
+	if (ColumnName == NAME_SyncPortSource) return Left.OutputReference == EMediaIOReferenceType::Input ? Left.ReferencePortIdentifier == Right.ReferencePortIdentifier : true;
+
+	return IdenticalProperty(ColumnName, Left.MediaConfiguration, Right.MediaConfiguration);
+}
+
+
+bool FMediaIOPermutationsSelectorBuilder::Less(FName ColumnName, const FMediaIOOutputConfiguration& Left, const FMediaIOOutputConfiguration& Right)
+{
+	if (ColumnName == NAME_OutputType) return (int32)Left.OutputType < (int32)Right.OutputType;
+	if (ColumnName == NAME_KeyPortSource) return Left.KeyPortIdentifier < Right.KeyPortIdentifier;
+	if (ColumnName == NAME_OutputReference) return (int32)Left.OutputReference < (int32)Right.OutputReference;
+	if (ColumnName == NAME_SyncPortSource) return Left.ReferencePortIdentifier < Right.ReferencePortIdentifier;
+
+	return Less(ColumnName, Left.MediaConfiguration, Right.MediaConfiguration);
+}
+
+
+FText FMediaIOPermutationsSelectorBuilder::GetLabel(FName ColumnName, const FMediaIOOutputConfiguration& Item)
+{
+	if (ColumnName == NAME_OutputType) return Item.OutputType == EMediaIOOutputType::Fill ? LOCTEXT("Fill", "Fill") : LOCTEXT("FillAndKey", "Fill and Key");
+	if (ColumnName == NAME_KeyPortSource) return MediaIOPermutationsSelectorBuilder::LinkTypeToPrettyText(Item.MediaConfiguration.MediaConnection.TransportType, Item.KeyPortIdentifier);
+	if (ColumnName == NAME_OutputReference)
+	{
+		switch (Item.OutputReference)
+		{
+		case EMediaIOReferenceType::FreeRun: return LOCTEXT("FreeRun", "Free Run");
+		case EMediaIOReferenceType::External: return LOCTEXT("External", "External");
+		case EMediaIOReferenceType::Input: return LOCTEXT("Input", "Input");
+		}
+	}
+	if (ColumnName == NAME_SyncPortSource) return MediaIOPermutationsSelectorBuilder::LinkTypeToPrettyText(Item.MediaConfiguration.MediaConnection.TransportType, Item.ReferencePortIdentifier);
+
+	return GetLabel(ColumnName, Item.MediaConfiguration);
+}
+
+
+FText FMediaIOPermutationsSelectorBuilder::GetTooltip(FName ColumnName, const FMediaIOOutputConfiguration& Item)
+{
+	if (ColumnName == NAME_OutputType) return LOCTEXT("OutputTypeTooltip", "Whether to output the fill or the fill and key.");
+	if (ColumnName == NAME_KeyPortSource) return FText::GetEmpty();
+	if (ColumnName == NAME_OutputReference) return LOCTEXT("OutputReferenceTooltip", "The Device output is synchronized with either its internal clock, an external reference, or an other input.");
+	if (ColumnName == NAME_SyncPortSource) return FText::GetEmpty();
+
+	return GetTooltip(ColumnName, Item.MediaConfiguration);
+}
+
+#undef LOCTEXT_NAMESPACE
diff --git a/Engine/Plugins/Media/MediaIOFramework/Source/MediaIOEditor/Public/Customizations/MediaIOConfigurationCustomization.h b/Engine/Plugins/Media/MediaIOFramework/Source/MediaIOEditor/Public/Customizations/MediaIOConfigurationCustomization.h
new file mode 100644
--- /dev/null
+++ b/Engine/Plugins/Media/MediaIOFramework/Source/MediaIOEditor/Public/Customizations/MediaIOConfigurationCustomization.h
@@ -0,0 +1,26 @@
+// Copyright 1998-2018 Epic Games, Inc. All Rights Reserved.
+
+#pragma once
+
+#include "Customizations/MediaIOCustomizationBase.h"
+#include "Input/Reply.h"
+#include "MediaIOCoreDefinitions.h"
+
+/**
+ * Implements a details view customization for the FMediaIOConfiguration
+ */
+class MEDIAIOEDITOR_API FMediaIOConfigurationCustomization : public FMediaIOCustomizationBase
+{
+public:
+	static TSharedRef<IPropertyTypeCustomization> MakeInstance();
+
+private:
+	virtual TAttribute<FText> GetContentText() override;
+	virtual TSharedRef<SWidget> HandleSourceComboButtonMenuContent() override;
+
+	void OnSelectionChanged(FMediaIOConfiguration SelectedItem);
+	FReply OnButtonClicked() const;
+
+	TWeakPtr<SWidget> PermutationSelector;
+	FMediaIOConfiguration SelectedConfiguration;
+};
diff --git a/Engine/Plugins/Media/MediaIOFramework/Source/MediaIOEditor/Public/Customizations/MediaIOCustomizationBase.h b/Engine/Plugins/Media/MediaIOFramework/Source/MediaIOEditor/Public/Customizations/MediaIOCustomizationBase.h
new file mode 100644
--- /dev/null
+++ b/Engine/Plugins/Media/MediaIOFramework/Source/MediaIOEditor/Public/Customizations/MediaIOCustomizationBase.h
@@ -0,0 +1,56 @@
+// Copyright 1998-2018 Epic Games, Inc. All Rights Reserved.
+
+#pragma once
+
+#include "CoreMinimal.h"
+#include "IPropertyTypeCustomization.h"
+#include "Misc/Attribute.h"
+#include "PropertyHandle.h"
+#include "Widgets/SWidget.h"
+
+/**
+ * Base implementation of different MediaIO details view customization
+ */
+class MEDIAIOEDITOR_API FMediaIOCustomizationBase : public IPropertyTypeCustomization
+{
+public:
+	/** IPropertyTypeCustomization interface */
+	virtual void CustomizeHeader(TSharedRef<class IPropertyHandle> InPropertyHandle, class FDetailWidgetRow& HeaderRow, IPropertyTypeCustomizationUtils& PropertyTypeCustomizationUtils) override;
+	virtual void CustomizeChildren(TSharedRef<class IPropertyHandle> InPropertyHandle, class IDetailChildrenBuilder& StructBuilder, IPropertyTypeCustomizationUtils& PropertyTypeCustomizationUtils) override;
+
+protected:
+	template<class T>
+	T* GetPropertyValueFromPropertyHandle()
+	{
+			UProperty* Property = MediaProperty->GetProperty();
+			check(Property);
+			check(Cast<UStructProperty>(Property));
+			check(Cast<UStructProperty>(Property)->Struct);
+			check(Cast<UStructProperty>(Property)->Struct->IsChildOf(T::StaticStruct()));
+
+			TArray<void*> RawData;
+			MediaProperty->AccessRawData(RawData);
+
+			check(RawData.Num() == 1);
+			T* MediaValue = reinterpret_cast<T*>(RawData[0]);
+			check(MediaValue);
+			return MediaValue;
+	}
+	template<class T>
+	void AssignValue(const T& NewValue) const
+	{
+		AssignValueImpl(reinterpret_cast<const void*>(&NewValue));
+	}
+
+	virtual TAttribute<FText> GetContentText() = 0;
+	virtual TSharedRef<SWidget> HandleSourceComboButtonMenuContent() = 0;
+	TSharedPtr<IPropertyHandle> GetMediaProperty() const { return MediaProperty; };
+
+	FName DeviceProviderName;
+
+private:
+	void AssignValueImpl(const void* NewValue) const;
+
+	/** Pointer to the property handle. */
+	TSharedPtr<IPropertyHandle> MediaProperty;
+};
diff --git a/Engine/Plugins/Media/MediaIOFramework/Source/MediaIOEditor/Public/Customizations/MediaIODeviceCustomization.h b/Engine/Plugins/Media/MediaIOFramework/Source/MediaIOEditor/Public/Customizations/MediaIODeviceCustomization.h
new file mode 100644
--- /dev/null
+++ b/Engine/Plugins/Media/MediaIOFramework/Source/MediaIOEditor/Public/Customizations/MediaIODeviceCustomization.h
@@ -0,0 +1,18 @@
+// Copyright 1998-2018 Epic Games, Inc. All Rights Reserved.
+
+#pragma once
+
+#include "Customizations/MediaIOCustomizationBase.h"
+
+/**
+ * Implements a details view customization for the FMediaIODevice
+ */
+class MEDIAIOEDITOR_API FMediaIODeviceCustomization : public FMediaIOCustomizationBase
+{
+public:
+	static TSharedRef<IPropertyTypeCustomization> MakeInstance();
+
+private:
+	virtual TAttribute<FText> GetContentText() override;
+	virtual TSharedRef<SWidget> HandleSourceComboButtonMenuContent() override;
+};
diff --git a/Engine/Plugins/Media/MediaIOFramework/Source/MediaIOEditor/Public/Customizations/MediaIOOutputConfigurationCustomization.h b/Engine/Plugins/Media/MediaIOFramework/Source/MediaIOEditor/Public/Customizations/MediaIOOutputConfigurationCustomization.h
new file mode 100644
--- /dev/null
+++ b/Engine/Plugins/Media/MediaIOFramework/Source/MediaIOEditor/Public/Customizations/MediaIOOutputConfigurationCustomization.h
@@ -0,0 +1,26 @@
+// Copyright 1998-2018 Epic Games, Inc. All Rights Reserved.
+
+#pragma once
+
+#include "Customizations/MediaIOCustomizationBase.h"
+#include "Input/Reply.h"
+#include "MediaIOCoreDefinitions.h"
+
+/**
+ * Implements a details view customization for the FMediaIOOutputConfiguration
+ */
+class MEDIAIOEDITOR_API FMediaIOOutputConfigurationCustomization : public FMediaIOCustomizationBase
+{
+public:
+	static TSharedRef<IPropertyTypeCustomization> MakeInstance();
+
+private:
+	virtual TAttribute<FText> GetContentText() override;
+	virtual TSharedRef<SWidget> HandleSourceComboButtonMenuContent() override;
+
+	void OnSelectionChanged(FMediaIOOutputConfiguration SelectedItem);
+	FReply OnButtonClicked() const;
+
+	TWeakPtr<SWidget> PermutationSelector;
+	FMediaIOOutputConfiguration SelectedConfiguration;
+};
diff --git a/Engine/Plugins/Media/MediaIOFramework/Source/MediaIOEditor/Public/MediaIOPermutationsSelectorBuilder.h b/Engine/Plugins/Media/MediaIOFramework/Source/MediaIOEditor/Public/MediaIOPermutationsSelectorBuilder.h
new file mode 100644
--- /dev/null
+++ b/Engine/Plugins/Media/MediaIOFramework/Source/MediaIOEditor/Public/MediaIOPermutationsSelectorBuilder.h
@@ -0,0 +1,42 @@
+// Copyright 1998-2018 Epic Games, Inc. All Rights Reserved.
+
+#pragma once
+
+#include "CoreMinimal.h"
+#include "MediaIOCoreDefinitions.h"
+
+
+struct MEDIAIOEDITOR_API FMediaIOPermutationsSelectorBuilder
+{
+	static const FName NAME_DeviceIdentifier;
+	static const FName NAME_TransportType;
+	static const FName NAME_QuadType;
+	static const FName NAME_Resolution;
+	static const FName NAME_Standard;
+	static const FName NAME_FrameRate;
+
+	static const FName NAME_OutputType;
+	static const FName NAME_KeyPortSource;
+	static const FName NAME_OutputReference;
+	static const FName NAME_SyncPortSource;
+
+	static bool IdenticalProperty(FName ColumnName, const FMediaIOConnection& Left, const FMediaIOConnection& Right);
+	static bool Less(FName ColumnName, const FMediaIOConnection& Left, const FMediaIOConnection& Right);
+	static FText GetLabel(FName ColumnName, const FMediaIOConnection& Item);
+	static FText GetTooltip(FName ColumnName, const FMediaIOConnection& Item);
+	
+	static bool IdenticalProperty(FName ColumnName, const FMediaIOMode& Left, const FMediaIOMode& Right);
+	static bool Less(FName ColumnName, const FMediaIOMode& Left, const FMediaIOMode& Right);
+	static FText GetLabel(FName ColumnName, const FMediaIOMode& Item);
+	static FText GetTooltip(FName ColumnName, const FMediaIOMode& Item);
+
+	static bool IdenticalProperty(FName ColumnName, const FMediaIOConfiguration& Left, const FMediaIOConfiguration& Right);
+	static bool Less(FName ColumnName, const FMediaIOConfiguration& Left, const FMediaIOConfiguration& Right);
+	static FText GetLabel(FName ColumnName, const FMediaIOConfiguration& Item);
+	static FText GetTooltip(FName ColumnName, const FMediaIOConfiguration& Item);
+
+	static bool IdenticalProperty(FName ColumnName, const FMediaIOOutputConfiguration& Left, const FMediaIOOutputConfiguration& Right);
+	static bool Less(FName ColumnName, const FMediaIOOutputConfiguration& Left, const FMediaIOOutputConfiguration& Right);
+	static FText GetLabel(FName ColumnName, const FMediaIOOutputConfiguration& Item);
+	static FText GetTooltip(FName ColumnName, const FMediaIOOutputConfiguration& Item);
+};
diff --git a/Engine/Plugins/Media/MediaIOFramework/Source/MediaIOEditor/Public/Widgets/SMediaPermutationsSelector.h b/Engine/Plugins/Media/MediaIOFramework/Source/MediaIOEditor/Public/Widgets/SMediaPermutationsSelector.h
new file mode 100644
--- /dev/null
+++ b/Engine/Plugins/Media/MediaIOFramework/Source/MediaIOEditor/Public/Widgets/SMediaPermutationsSelector.h
@@ -0,0 +1,198 @@
+// Copyright 1998-2018 Epic Games, Inc. All Rights Reserved.
+
+#pragma once
+
+#include "CoreMinimal.h"
+#include "Misc/Optional.h"
+#include "Widgets/DeclarativeSyntaxSupport.h"
+#include "Widgets/SCompoundWidget.h"
+
+
+/**
+ * A widget that let you select a single permutation from a list . It groups the values into categories and removes duplicates inside that category.
+ *
+ * A trivial use case appears below:
+ *
+ *   struct FMyItem { int32 A; int32 B; }
+ *   struct FMyBuilder
+ *   {
+ *      static const FName ColumnA = "A";
+ *      static const FName ColumnB = "B";
+ *      static bool IdenticalProperty(FName ColumnName, const FMyItem& Left, const FMyItem& Right) { return ColumnName == ColumnA ? Left.A == Right.A : Left.B == Right.B; }
+ *      static bool Less(FName ColumnName, const FMyItem& Left, const FMyItem& Right) { return ColumnName == ColumnA ? Left.A < Right.A : Left.B < Right.B; }
+ *      static FText GetLabel(FName ColumnName, const FMyItem& Item) { return ColumnName == ColumnA ? FText::AsNumber(Item.A) : FText::AsNumber(Item.B); }
+ *      static FText GetTooltip(FName ColumnName, const FMyItem& Item) { return LOCTEXT("Tooltip", "Tooltip"); }
+ *   };
+ *
+ *   TArray< FMyItem > Items;
+ *   FMyItem Value1;
+ *   Value1.A = 1; Value1.B = 2;
+ *   Items.Add( Value1 );
+ *   FMyItem Value2;
+ *   Value2.A = 1; Value2.B = 3;
+ *   Items.Add( Value2 );
+ *
+ *   using TSelection = SMediaPermutationsSelector< FMyItem, FMyBuilder >;
+ *   SNew( TSelection )
+ *     .PermutationsSource( MoveTemp(Items) )
+ *     .SelectedPermutation( Value2 )
+ *     + TSelection::Column( FMyBuilder::A );
+ *     .Label( LOCTEXT("ExampleA", "The A") )
+ *     + TSelection::Column( FMyBuilder::B );
+ *     .Label( LOCTEXT("ExampleB", "The B") )
+ *
+ * In the example, we make all 2 columns. One for A and one for B.
+ * The first column will have 1 element: "1". The second column will have 2 elements: "2", "3"
+ *
+ */
+
+class SHorizontalBox;
+
+template <typename ItemType>
+class MEDIAIOEDITOR_API TMediaPermutationsSelectorBuilder
+{
+public:
+	static bool IdenticalProperty(FName ColumnName, ItemType Left, ItemType Right) { return Left == Right; }
+	static bool Less(FName ColumnName, ItemType Left, ItemType Right) { return Left < Right; }
+	static FText GetLabel(FName ColumnName, ItemType Item) { return FText::FromName(ColumnName); }
+	static FText GetTooltip(FName ColumnName, ItemType Item) { return FText::FromName(ColumnName); }
+};
+
+template <typename ItemType, typename ItemBuilder = TMediaPermutationsSelectorBuilder<ItemType>>
+class SMediaPermutationsSelector : public SCompoundWidget
+{
+public:
+	using ThisClass = SMediaPermutationsSelector<ItemType, ItemBuilder>;
+
+	/** A delegate type invoked when the selection changes. */
+	DECLARE_DELEGATE_OneParam(FOnSelectionChanged, ItemType /* NewItemSelected*/);
+
+	/**
+	 * A delegate type invoked when we fill a column and wants to check if it should be visible.
+	 * @param ColumnName the name of the ColumnName
+	 * @param UniqueItemInColumn  List of items used to generate the column. Multiples items may be shared the value but only one will be in that list.
+	 * @return true if the column should be visible
+	 */
+	DECLARE_DELEGATE_RetVal_TwoParams(bool, FIsColumnVisible, FName /*ColumnName*/, const TArray<ItemType>& /*UniquePermutationsForThisColumn */);
+
+public:
+	/** Describes a single column */
+	class FColumn
+	{
+	public:
+		SLATE_BEGIN_ARGS(FColumn)
+		{}
+			/** A unique ID for this property, so that it can be saved and restored. */
+			SLATE_ARGUMENT(FName, ColumnName)
+			/** Text to use as the Column header. */
+			SLATE_ATTRIBUTE(FText, Label)
+			/** Text to use as the Column tooltip. */
+			SLATE_ATTRIBUTE(FText, Tooltip)
+			/** Delegate to invoke when build the column and check the visibility. */
+			SLATE_EVENT(FIsColumnVisible, IsColumnVisible)
+
+		SLATE_END_ARGS()
+
+		FColumn(const FArguments& InArgs)
+			: ColumnName(InArgs._ColumnName)
+			, Label(InArgs._Label)
+			, Tooltip(InArgs._Tooltip)
+			, IsColumnVisible(InArgs._IsColumnVisible)
+		{ }
+
+	public:
+		/** A unique ID for this property, so that it can be saved and restored. */
+		FName ColumnName;
+
+		/** Text to use as the Column header. */
+		TAttribute< FText > Label;
+
+		/** Text to use as the Column tooltip. */
+		TAttribute< FText > Tooltip;
+
+		/** Widget created by this menu. */
+		TSharedPtr< SWidget > Widget;
+
+		/** Delegate to invoke when build the column and check the visibility. */
+		FIsColumnVisible IsColumnVisible;
+	};
+
+public:
+	SLATE_BEGIN_ARGS(ThisClass)
+		: _ColumnHeight(200)
+	{}
+		/** Array of columns */
+		SLATE_SUPPORTS_SLOT_WITH_ARGS(FColumn)
+		/** Array of data items that we are displaying */
+		SLATE_ARGUMENT(TArray<ItemType>, PermutationsSource)
+		/** Default selected item in ItemsSource */
+		SLATE_ARGUMENT(TOptional<ItemType>, SelectedPermutation)
+		/** Desired height of the columns */
+		SLATE_ATTRIBUTE(FOptionalSize, ColumnHeight)
+		/** Override the "apply" button widget */
+		SLATE_ARGUMENT(TSharedPtr<SWidget>, OverrideButtonWidget)
+		/** Delegate to invoke when the button is clicked. */
+		SLATE_EVENT(FOnClicked, OnButtonClicked)
+		/** Delegate to invoke when selection changes. */
+		SLATE_EVENT(FOnSelectionChanged, OnSelectionChanged)
+	SLATE_END_ARGS()
+
+	void Construct(const FArguments& InArgs);
+
+	/** Create a column with a specified ColumnId */
+	static typename FColumn::FArguments Column(const FName& InColumnName);
+
+	/** Get a copy of the current selected item. */
+	ItemType GetSelectedItem() const;
+
+private:
+	/** Array of data items that we are displaying. */
+	TArray<ItemType> PermutationsSource;
+
+	/** Index of the selected item in ItemsSource. It will always be valid. */
+	int32 SelectedPermutationIndex;
+
+	/** Delegate to invoke when selection changes. */
+	FOnSelectionChanged OnSelectionChanged;
+
+	/** Columns information. */
+	TIndirectArray<FColumn> PropertyColumns;
+
+	/** Box used as container for the radio button menu. */
+	TSharedPtr<SHorizontalBox> ColumnContainer;
+
+private:
+	void BuildColumns(int32 StartIndex);
+	void ItemSelected(int32 UniqueItemIndexSelected, int32 ColumnSelected);
+	TArray<int32> GenerateItemIndexes() const;
+	bool IsColumnVisible(const FColumn& Column, const TArray<int32>& UniqueItemIndexes) const;
+
+private:
+	bool IdenticalProperty(FName ColumnName, int32 LeftItemIndex, int32 RightItemIndex)
+	{
+		check(PermutationsSource.IsValidIndex(LeftItemIndex));
+		check(PermutationsSource.IsValidIndex(RightItemIndex));
+		return ItemBuilder::IdenticalProperty(ColumnName, PermutationsSource[LeftItemIndex], PermutationsSource[RightItemIndex]);
+	}
+
+	bool Less(FName ColumnName, int32 LeftItemIndex, int32 RightItemIndex)
+	{
+		check(PermutationsSource.IsValidIndex(LeftItemIndex));
+		check(PermutationsSource.IsValidIndex(RightItemIndex));
+		return ItemBuilder::Less(ColumnName, PermutationsSource[LeftItemIndex], PermutationsSource[RightItemIndex]);
+	}
+
+	FText GetLabel(FName ColumnName, int32 ItemIndex)
+	{
+		check(PermutationsSource.IsValidIndex(ItemIndex));
+		return ItemBuilder::GetLabel(ColumnName, PermutationsSource[ItemIndex]);
+	}
+
+	FText GetTooltip(FName ColumnName, int32 ItemIndex)
+	{
+		check(PermutationsSource.IsValidIndex(ItemIndex));
+		return ItemBuilder::GetTooltip(ColumnName, PermutationsSource[ItemIndex]);
+	}
+};
+
+#include "Widgets/SMediaPermutationsSelector.inl"
diff --git a/Engine/Plugins/Media/MediaIOFramework/Source/MediaIOEditor/Public/Widgets/SMediaPermutationsSelector.inl b/Engine/Plugins/Media/MediaIOFramework/Source/MediaIOEditor/Public/Widgets/SMediaPermutationsSelector.inl
new file mode 100644
--- /dev/null
+++ b/Engine/Plugins/Media/MediaIOFramework/Source/MediaIOEditor/Public/Widgets/SMediaPermutationsSelector.inl
@@ -0,0 +1,303 @@
+// Copyright 1998-2018 Epic Games, Inc. All Rights Reserved.
+
+#include "Framework/MultiBox/MultiBoxBuilder.h"
+#include "Widgets/SBoxPanel.h"
+#include "Widgets/Input/SButton.h"
+#include "Widgets/Layout/SBox.h"
+
+template <typename ItemType, typename ItemBuilder>
+typename SMediaPermutationsSelector<ItemType, ItemBuilder>::FColumn::FArguments SMediaPermutationsSelector<ItemType, ItemBuilder>::Column(const FName& InColumnName)
+{
+	typename FColumn::FArguments NewArgs;
+	NewArgs._ColumnName = InColumnName;
+	return NewArgs;
+}
+
+template <typename ItemType, typename ItemBuilder>
+ItemType SMediaPermutationsSelector<ItemType, ItemBuilder>::GetSelectedItem() const
+{
+	return PermutationsSource.IsValidIndex(SelectedPermutationIndex) ? PermutationsSource[SelectedPermutationIndex] : ItemType();
+}
+
+template <typename ItemType, typename ItemBuilder>
+void SMediaPermutationsSelector<ItemType, ItemBuilder>::Construct(const FArguments& InArgs)
+{
+	SWidget::Construct(InArgs._ToolTipText, InArgs._ToolTip, InArgs._Cursor, InArgs._IsEnabled, InArgs._Visibility, InArgs._RenderOpacity, InArgs._RenderTransform, InArgs._RenderTransformPivot, InArgs._Tag, InArgs._ForceVolatile, InArgs._Clipping, InArgs.MetaData);
+
+	PermutationsSource = InArgs._PermutationsSource;
+	SelectedPermutationIndex = INDEX_NONE;
+	OnSelectionChanged = InArgs._OnSelectionChanged;
+
+	if (InArgs._SelectedPermutation.IsSet())
+	{
+		SelectedPermutationIndex = PermutationsSource.IndexOfByKey(InArgs._SelectedPermutation.GetValue());
+	}
+
+	// Copy all the column info from the declaration
+	PropertyColumns.Empty(InArgs.Slots.Num());
+	for (int32 SlotIndex = 0; SlotIndex < InArgs.Slots.Num(); ++SlotIndex)
+	{
+		FColumn* const Column = InArgs.Slots[SlotIndex];
+		if (Column && Column->ColumnName != NAME_None)
+		{
+			PropertyColumns.Add(Column);
+		}
+	}
+
+	SAssignNew(ColumnContainer, SHorizontalBox);
+
+	BuildColumns(0);
+
+	ChildSlot
+	[
+		SNew(SBox)
+		.HeightOverride(InArgs._ColumnHeight)
+		[
+			SNew(SVerticalBox)
+			+ SVerticalBox::Slot()
+			.FillHeight(1.f)
+			[
+				ColumnContainer.ToSharedRef()
+			]
+			+ SVerticalBox::Slot()
+			.AutoHeight()
+			[
+				SNew(SHorizontalBox)
+				+ SHorizontalBox::Slot()
+				.FillWidth(1.f)
+				[
+					SNullWidget::NullWidget
+				]
+				+ SHorizontalBox::Slot()
+				.Padding(4.f)
+				.AutoWidth()
+				[
+					InArgs._OverrideButtonWidget.IsValid() ?
+						InArgs._OverrideButtonWidget.ToSharedRef()
+						:
+						SNew(SButton)
+						.VAlign(VAlign_Center)
+						.HAlign(HAlign_Center)
+						.Text(NSLOCTEXT("MediaPlayerEditor", "ApplyLabel", "Apply"))
+						.OnClicked(InArgs._OnButtonClicked)
+				]
+			]
+		]
+	];
+}
+
+template <typename ItemType, typename ItemBuilder>
+void SMediaPermutationsSelector<ItemType, ItemBuilder>::BuildColumns(int32 StartIndex)
+{
+	const int32 NumberOfItems = PermutationsSource.Num();
+	if (!PropertyColumns.IsValidIndex(StartIndex) || NumberOfItems == 0)
+	{
+		return;
+	}
+
+	if (SelectedPermutationIndex == INDEX_NONE)
+	{
+		SelectedPermutationIndex = 0;
+		OnSelectionChanged.ExecuteIfBound(PermutationsSource[SelectedPermutationIndex]);
+	}
+
+	// Get only the valid item for that selected input
+	TArray<int32> AllValidItemIndexes = GenerateItemIndexes();
+
+	// Clean the items list. AllValidItems should contains only the valid entry up to category StartIndex.
+	for (int32 ColumnIndex = 0; ColumnIndex < StartIndex; ++ColumnIndex)
+	{
+		const FColumn& Column = PropertyColumns[ColumnIndex];
+		for (int32 ItemIndex = AllValidItemIndexes.Num() - 1; ItemIndex >= 0; --ItemIndex)
+		{
+			if (!IdenticalProperty(Column.ColumnName, SelectedPermutationIndex, AllValidItemIndexes[ItemIndex]))
+			{
+				AllValidItemIndexes.RemoveAtSwap(ItemIndex);
+			}
+		}
+	}
+
+	// ReBuild the items
+	for (int32 ColumnIndex = StartIndex; ColumnIndex < PropertyColumns.Num(); ++ColumnIndex)
+	{
+		FColumn& Column = PropertyColumns[ColumnIndex];
+		if (Column.Widget.IsValid())
+		{
+			ColumnContainer->RemoveSlot(Column.Widget.ToSharedRef());
+			Column.Widget.Reset();
+		}
+
+		// Gather the unique item and remove all that do not match with the current selected item
+		TArray<int32> UniqueItemsForColumnIndexes;
+		const FName ColumnName = Column.ColumnName;
+
+		for (int32 ItemIndex = AllValidItemIndexes.Num() - 1; ItemIndex >= 0; --ItemIndex)
+		{
+			bool bFound = false;
+			for (int32 UniqueIndex : UniqueItemsForColumnIndexes)
+			{
+				if (IdenticalProperty(ColumnName, UniqueIndex, AllValidItemIndexes[ItemIndex]))
+				{
+					bFound = true;
+					break;
+				}
+			}
+			if (!bFound)
+			{
+				UniqueItemsForColumnIndexes.Add(AllValidItemIndexes[ItemIndex]);
+			}
+
+			if (!IdenticalProperty(ColumnName, SelectedPermutationIndex, AllValidItemIndexes[ItemIndex]))
+			{
+				AllValidItemIndexes.RemoveAtSwap(ItemIndex);
+			}
+		}
+
+		// Only show the column if the user desire it
+		if (!IsColumnVisible(Column, UniqueItemsForColumnIndexes))
+		{
+			continue;
+		}
+
+		// Sort the column items
+		UniqueItemsForColumnIndexes.Sort([this, ColumnName](int32 Left, int32 Right) { return Less(ColumnName, Left, Right); });
+
+		// Build the radio buttons
+		FMenuBuilder MenuBuilder(false, nullptr);
+		for (int32 UniqueItemIndex : UniqueItemsForColumnIndexes)
+		{
+			MenuBuilder.AddMenuEntry(
+				GetLabel(ColumnName, UniqueItemIndex),
+				GetTooltip(ColumnName, UniqueItemIndex),
+				FSlateIcon(),
+				FUIAction(
+					FExecuteAction::CreateLambda([this, UniqueItemIndex, ColumnIndex] { ItemSelected(UniqueItemIndex, ColumnIndex); }),
+					FCanExecuteAction(),
+					FIsActionChecked::CreateLambda([this, UniqueItemIndex, ColumnName] { return IdenticalProperty(ColumnName, SelectedPermutationIndex, UniqueItemIndex); })
+				),
+				NAME_None,
+				EUserInterfaceActionType::RadioButton
+			);
+		}
+
+		// Create the widget
+		SAssignNew(Column.Widget, SVerticalBox)
+		+ SVerticalBox::Slot()
+		.AutoHeight()
+		.HAlign(HAlign_Left)
+		.Padding(2, 0, 5, 0)
+		[
+			SNew(STextBlock)
+			.Text(Column.Label)
+		]
+		+ SVerticalBox::Slot()
+		.FillHeight(1.f)
+		[
+			MenuBuilder.MakeWidget()
+		];
+
+		// Attach the widget
+		ColumnContainer->AddSlot()
+		.AutoWidth()
+		[
+			Column.Widget.ToSharedRef()
+		];
+	}
+}
+
+template <typename ItemType, typename ItemBuilder>
+void SMediaPermutationsSelector<ItemType, ItemBuilder>::ItemSelected(int32 UniqueItemIndex, int32 ColumnIndex)
+{
+	const int32 NumberOfItems = PermutationsSource.Num();
+	if (!PropertyColumns.IsValidIndex(ColumnIndex) || NumberOfItems == 0)
+	{
+		return;
+	}
+
+	int32 PreviousSelectedItemIndex = SelectedPermutationIndex;
+	SelectedPermutationIndex = INDEX_NONE;
+	
+	// Get only the valid item for that selected input
+	TArray<int32> AllValidItemIndexes = GenerateItemIndexes();
+	
+	int32 MaxColumns = FMath::Min(ColumnIndex + 1, PropertyColumns.Num());
+	for (int32 Index = 0; Index < MaxColumns; ++Index)
+	{
+		const FColumn& Column = PropertyColumns[Index];
+		for (int32 ItemIndex = AllValidItemIndexes.Num() - 1; ItemIndex >= 0; --ItemIndex)
+		{
+			if (!IdenticalProperty(Column.ColumnName, UniqueItemIndex, AllValidItemIndexes[ItemIndex]))
+			{
+				AllValidItemIndexes.RemoveAtSwap(ItemIndex);
+			}
+		}
+	}
+	
+	if (AllValidItemIndexes.Num() > 0)
+	{
+		SelectedPermutationIndex = AllValidItemIndexes[0];
+	}
+	
+	// Try to find something that matches what we used to have.
+	for (int32 Index = MaxColumns; Index < PropertyColumns.Num(); ++Index)
+	{
+		const FColumn& Column = PropertyColumns[Index];
+		for (int32 ItemIndex = AllValidItemIndexes.Num() - 1; ItemIndex >= 0; --ItemIndex)
+		{
+			if (!IdenticalProperty(Column.ColumnName, PreviousSelectedItemIndex, AllValidItemIndexes[ItemIndex]))
+			{
+				AllValidItemIndexes.RemoveAtSwap(ItemIndex);
+			}
+		}
+	
+		if (AllValidItemIndexes.Num() > 0)
+		{
+			SelectedPermutationIndex = AllValidItemIndexes[0];
+		}
+	}
+	
+	// There should always be one selected item
+	if (SelectedPermutationIndex == INDEX_NONE)
+	{
+		SelectedPermutationIndex = UniqueItemIndex;
+	}
+	
+	BuildColumns(ColumnIndex + 1);
+		
+	if (SelectedPermutationIndex != PreviousSelectedItemIndex)
+	{
+		OnSelectionChanged.ExecuteIfBound(PermutationsSource[SelectedPermutationIndex]);
+	}
+}
+
+template <typename ItemType, typename ItemBuilder>
+TArray<int32> SMediaPermutationsSelector<ItemType, ItemBuilder>::GenerateItemIndexes() const
+{
+	TArray<int32> Result;
+	const int32 NumberOfItems = PermutationsSource.Num();
+	Result.Reserve(NumberOfItems);
+	for (int32 Index = 0; Index < NumberOfItems; ++Index)
+	{
+		Result.Add(Index);
+	}
+	return MoveTemp(Result);
+}
+
+template <typename ItemType, typename ItemBuilder>
+bool SMediaPermutationsSelector<ItemType, ItemBuilder>::IsColumnVisible(const FColumn& Column, const TArray<int32>& UniqueItemIndexes) const
+{
+	bool bResult = true;
+	if (Column.IsColumnVisible.IsBound())
+	{
+		TArray<ItemType> UniqueItems;
+		UniqueItems.Empty(UniqueItemIndexes.Num());
+		for (int32 ItemIndex : UniqueItemIndexes)
+		{
+			UniqueItems.Add(PermutationsSource[ItemIndex]);
+		}
+
+		bResult = Column.IsColumnVisible.Execute(Column.ColumnName, UniqueItems);
+	}
+
+	return bResult;
+}
diff --git a/Engine/Plugins/Media/MediaPlayerEditor/Source/MediaPlayerEditor/Public/Widgets/SMediaPermutationsSelector.h b/Engine/Plugins/Media/MediaPlayerEditor/Source/MediaPlayerEditor/Public/Widgets/SMediaPermutationsSelector.h
deleted file mode 100644
diff --git a/Engine/Plugins/Media/MediaPlayerEditor/Source/MediaPlayerEditor/Public/Widgets/SMediaPermutationsSelector.inl b/Engine/Plugins/Media/MediaPlayerEditor/Source/MediaPlayerEditor/Public/Widgets/SMediaPermutationsSelector.inl
deleted file mode 100644
diff --git a/Engine/Source/Runtime/MediaIOCore/MediaIOCore.Build.cs b/Engine/Source/Runtime/MediaIOCore/MediaIOCore.Build.cs
deleted file mode 100644
diff --git a/Engine/Source/Runtime/MediaIOCore/Private/MediaCapture.cpp b/Engine/Source/Runtime/MediaIOCore/Private/MediaCapture.cpp
deleted file mode 100644
diff --git a/Engine/Source/Runtime/MediaIOCore/Private/MediaIOCoreModule.cpp b/Engine/Source/Runtime/MediaIOCore/Private/MediaIOCoreModule.cpp
deleted file mode 100644
diff --git a/Engine/Source/Runtime/MediaIOCore/Private/MediaIOCoreModule.h b/Engine/Source/Runtime/MediaIOCore/Private/MediaIOCoreModule.h
deleted file mode 100644
diff --git a/Engine/Source/Runtime/MediaIOCore/Private/MediaOutput.cpp b/Engine/Source/Runtime/MediaIOCore/Private/MediaOutput.cpp
deleted file mode 100644
diff --git a/Engine/Source/Runtime/MediaIOCore/Private/Player/MediaIOCorePlayerBase.cpp b/Engine/Source/Runtime/MediaIOCore/Private/Player/MediaIOCorePlayerBase.cpp
deleted file mode 100644
diff --git a/Engine/Source/Runtime/MediaIOCore/Private/Player/MediaIOCoreSamples.cpp b/Engine/Source/Runtime/MediaIOCore/Private/Player/MediaIOCoreSamples.cpp
deleted file mode 100644
diff --git a/Engine/Source/Runtime/MediaIOCore/Private/Shared/MediaIOCoreEncodeTime.cpp b/Engine/Source/Runtime/MediaIOCore/Private/Shared/MediaIOCoreEncodeTime.cpp
deleted file mode 100644
diff --git a/Engine/Source/Runtime/MediaIOCore/Private/Shared/MediaIOCoreWaitVsyncThread.cpp b/Engine/Source/Runtime/MediaIOCore/Private/Shared/MediaIOCoreWaitVsyncThread.cpp
deleted file mode 100644
diff --git a/Engine/Source/Runtime/MediaIOCore/Public/IMediaIOCoreHardwareSync.h b/Engine/Source/Runtime/MediaIOCore/Public/IMediaIOCoreHardwareSync.h
deleted file mode 100644
diff --git a/Engine/Source/Runtime/MediaIOCore/Public/MediaCapture.h b/Engine/Source/Runtime/MediaIOCore/Public/MediaCapture.h
deleted file mode 100644
diff --git a/Engine/Source/Runtime/MediaIOCore/Public/MediaIOCoreAudioSampleBase.h b/Engine/Source/Runtime/MediaIOCore/Public/MediaIOCoreAudioSampleBase.h
deleted file mode 100644
diff --git a/Engine/Source/Runtime/MediaIOCore/Public/MediaIOCoreBinarySampleBase.h b/Engine/Source/Runtime/MediaIOCore/Public/MediaIOCoreBinarySampleBase.h
deleted file mode 100644
diff --git a/Engine/Source/Runtime/MediaIOCore/Public/MediaIOCoreEncodeTime.h b/Engine/Source/Runtime/MediaIOCore/Public/MediaIOCoreEncodeTime.h
deleted file mode 100644
diff --git a/Engine/Source/Runtime/MediaIOCore/Public/MediaIOCorePlayerBase.h b/Engine/Source/Runtime/MediaIOCore/Public/MediaIOCorePlayerBase.h
deleted file mode 100644
diff --git a/Engine/Source/Runtime/MediaIOCore/Public/MediaIOCoreSamples.h b/Engine/Source/Runtime/MediaIOCore/Public/MediaIOCoreSamples.h
deleted file mode 100644
diff --git a/Engine/Source/Runtime/MediaIOCore/Public/MediaIOCoreTextureSampleBase.h b/Engine/Source/Runtime/MediaIOCore/Public/MediaIOCoreTextureSampleBase.h
deleted file mode 100644
diff --git a/Engine/Source/Runtime/MediaIOCore/Public/MediaIOCoreWaitVsyncThread.h b/Engine/Source/Runtime/MediaIOCore/Public/MediaIOCoreWaitVsyncThread.h
deleted file mode 100644
diff --git a/Engine/Source/Runtime/MediaIOCore/Public/MediaOutput.h b/Engine/Source/Runtime/MediaIOCore/Public/MediaOutput.h
deleted file mode 100644
