Warn the user when a MediaCapture is about to be destroyed and the capture was not stopped.

diff --git a/Engine/Plugins/Media/MediaIOFramework/Source/MediaIOCore/Private/MediaCapture.cpp b/Engine/Plugins/Media/MediaIOFramework/Source/MediaIOCore/Private/MediaCapture.cpp
--- a/Engine/Plugins/Media/MediaIOFramework/Source/MediaIOCore/Private/MediaCapture.cpp
+++ b/Engine/Plugins/Media/MediaIOFramework/Source/MediaIOCore/Private/MediaCapture.cpp
@@ -72,6 +72,10 @@ UMediaCapture::UMediaCapture(const FObjectInitializer& ObjectInitializer)
 
 void UMediaCapture::BeginDestroy()
 {
+	if (MediaState == EMediaCaptureState::Capturing || MediaState == EMediaCaptureState::Preparing)
+	{
+		UE_LOG(LogMediaIOCore, Warning, TEXT("%s will be destroyed and the capture was not stopped."), *GetName());
+	}
 	StopCapture(false);
 
 	Super::BeginDestroy();
@@ -238,6 +242,12 @@ bool UMediaCapture::UpdateTextureRenderTarget2D(UTextureRenderTarget2D * InRende
 void UMediaCapture::StopCapture(bool bAllowPendingFrameToBeProcess)
 {
 	check(IsInGameThread());
+
+	if (MediaState != EMediaCaptureState::StopRequested && MediaState != EMediaCaptureState::Capturing)
+	{
+		bAllowPendingFrameToBeProcess = false;
+	}
+
 	if (bAllowPendingFrameToBeProcess)
 	{
 		if (MediaState != EMediaCaptureState::Stopped && MediaState != EMediaCaptureState::StopRequested)
