Upgrade SMediaPermutationsSelector.

diff --git a/Engine/Plugins/Media/MediaIOFramework/Source/MediaIOEditor/Public/Widgets/SMediaPermutationsSelector.inl b/Engine/Plugins/Media/MediaIOFramework/Source/MediaIOEditor/Public/Widgets/SMediaPermutationsSelector.inl
--- a/Engine/Plugins/Media/MediaIOFramework/Source/MediaIOEditor/Public/Widgets/SMediaPermutationsSelector.inl
+++ b/Engine/Plugins/Media/MediaIOFramework/Source/MediaIOEditor/Public/Widgets/SMediaPermutationsSelector.inl
@@ -22,7 +22,7 @@ ItemType SMediaPermutationsSelector<ItemType, ItemBuilder>::GetSelectedItem() co
 template <typename ItemType, typename ItemBuilder>
 void SMediaPermutationsSelector<ItemType, ItemBuilder>::Construct(const FArguments& InArgs)
 {
-	SWidget::Construct(InArgs._ToolTipText, InArgs._ToolTip, InArgs._Cursor, InArgs._IsEnabled, InArgs._Visibility, InArgs._RenderOpacity, InArgs._RenderTransform, InArgs._RenderTransformPivot, InArgs._Tag, InArgs._ForceVolatile, InArgs._Clipping, InArgs.MetaData);
+	SWidget::Construct(InArgs._ToolTipText, InArgs._ToolTip, InArgs._Cursor, InArgs._IsEnabled, InArgs._Visibility, InArgs._RenderOpacity, InArgs._RenderTransform, InArgs._RenderTransformPivot, InArgs._Tag, InArgs._ForceVolatile, InArgs._Clipping, InArgs._FlowDirectionPreference, InArgs.MetaData);
 
 	PermutationsSource = InArgs._PermutationsSource;
 	SelectedPermutationIndex = INDEX_NONE;
