Fix Timecode synchronizer which never resolve with SDI input.

Fix MediaIoCore assert in uninitialized array.

Improve MediaCapture thread-safeness

diff --git a/Engine/Plugins/Media/MediaIOFramework/Source/MediaIOCore/Private/MediaCapture.cpp b/Engine/Plugins/Media/MediaIOFramework/Source/MediaIOCore/Private/MediaCapture.cpp
--- a/Engine/Plugins/Media/MediaIOFramework/Source/MediaIOCore/Private/MediaCapture.cpp
+++ b/Engine/Plugins/Media/MediaIOFramework/Source/MediaIOCore/Private/MediaCapture.cpp
@@ -94,7 +94,7 @@ UMediaCapture::UMediaCapture(const FObjectInitializer& ObjectInitializer)
 	, ConversionOperation(EMediaCaptureConversionOperation::NONE)
 	, MediaOutputName(TEXT("[undefined]"))
 	, bResolvedTargetInitialized(false)
-	, bWaitingForResolveCommandExecution(false)
+	, WaitingForResolveCommandExecutionCounter(0)
 {
 }
 
@@ -348,7 +348,7 @@ void UMediaCapture::StopCapture(bool bAllowPendingFrameToBeProcess)
 
 			FCoreDelegates::OnEndFrame.RemoveAll(this);
 
-			while (bWaitingForResolveCommandExecution || !bResolvedTargetInitialized)
+			while (WaitingForResolveCommandExecutionCounter != 0 || !bResolvedTargetInitialized)
 			{
 				FlushRenderingCommands();
 			}
@@ -378,7 +378,7 @@ void UMediaCapture::SetMediaOutput(UMediaOutput* InMediaOutput)
 
 bool UMediaCapture::HasFinishedProcessing() const
 {
-	return bWaitingForResolveCommandExecution == false
+	return WaitingForResolveCommandExecutionCounter == 0
 		|| MediaState == EMediaCaptureState::Error
 		|| MediaState == EMediaCaptureState::Stopped;
 }
@@ -481,32 +481,23 @@ void UMediaCapture::OnEndFrame_GameThread()
 		CapturingFrame->UserData = GetCaptureFrameUserData_GameThread();
 	}
 
-	bWaitingForResolveCommandExecution = true;
-
 	// RenderCommand to be executed on the RenderThread
-	auto RenderCommand = [this](FRHICommandListImmediate& RHICmdList, FCaptureFrame* InCapturingFrame, FCaptureFrame* InReadyFrame)
+	auto RenderCommand = [this](FRHICommandListImmediate& RHICmdList, FCaptureFrame* InCapturingFrame, FCaptureFrame* InReadyFrame, TWeakPtr<FSceneViewport> InViewport, FTextureRenderTargetResource* InTextureRenderTargetResource)
 	{
 		FTexture2DRHIRef SourceTexture;
 		{
-			FScopeLock Lock(&AccessingCapturingSource);
-
-			UTextureRenderTarget2D* InCapturingRenderTarget = CapturingRenderTarget;
-			TSharedPtr<FSceneViewport> InSceneViewportPtr = CapturingSceneViewport.Pin();
-
-			if (InSceneViewportPtr.IsValid())
+			TSharedPtr<FSceneViewport> SceneViewportPtr = InViewport.Pin();
+			if (SceneViewportPtr)
 			{
-				SourceTexture = InSceneViewportPtr->GetRenderTargetTexture();
-				if (!SourceTexture.IsValid() && InSceneViewportPtr->GetViewportRHI())
+				SourceTexture = SceneViewportPtr->GetRenderTargetTexture();
+				if (!SourceTexture.IsValid() && SceneViewportPtr->GetViewportRHI())
 				{
-					SourceTexture = RHICmdList.GetViewportBackBuffer(InSceneViewportPtr->GetViewportRHI());
+					SourceTexture = RHICmdList.GetViewportBackBuffer(SceneViewportPtr->GetViewportRHI());
 				}
 			}
-			else if (InCapturingRenderTarget)
+			else if (InTextureRenderTargetResource && InTextureRenderTargetResource->GetTextureRenderTarget2DResource())
 			{
-				if (InCapturingRenderTarget->GetRenderTargetResource() != nullptr && InCapturingRenderTarget->GetRenderTargetResource()->GetTextureRenderTarget2DResource() != nullptr)
-				{
-					SourceTexture = InCapturingRenderTarget->GetRenderTargetResource()->GetTextureRenderTarget2DResource()->GetTextureRHI();
-				}
+				SourceTexture = InTextureRenderTargetResource->GetTextureRenderTarget2DResource()->GetTextureRHI();
 			}
 		}
 
@@ -703,17 +694,29 @@ void UMediaCapture::OnEndFrame_GameThread()
 			RHICmdList.UnmapStagingSurface(InReadyFrame->ReadbackTexture);
 		}
 
-		bWaitingForResolveCommandExecution = false;
+		--WaitingForResolveCommandExecutionCounter;
 	};
 
+	FTextureRenderTargetResource* TextureRenderTargetResource = nullptr;
+	{
+		FScopeLock Lock(&AccessingCapturingSource);
+
+		if(CapturingRenderTarget)
+		{
+			TextureRenderTargetResource = CapturingRenderTarget->GameThread_GetRenderTargetResource();
+		}
+	}
 
-	ENQUEUE_UNIQUE_RENDER_COMMAND_THREEPARAMETER(
+	++WaitingForResolveCommandExecutionCounter;
+	ENQUEUE_UNIQUE_RENDER_COMMAND_FIVEPARAMETER(
 		MediaOutputCaptureFrameCreateTexture,
 		FCaptureFrame*, InCapturingFrame, CapturingFrame,
 		FCaptureFrame*, InPreviousFrame, ReadyFrame,
+		TWeakPtr<FSceneViewport>, InViewport, CapturingSceneViewport,
+		FTextureRenderTargetResource*, InTextureRenderTargetResource, TextureRenderTargetResource,
 		decltype(RenderCommand), InRenderCommand, RenderCommand,
 		{
-			InRenderCommand(RHICmdList, InCapturingFrame, InPreviousFrame);
+			InRenderCommand(RHICmdList, InCapturingFrame, InPreviousFrame, InViewport, InTextureRenderTargetResource);
 		});
 }
 
diff --git a/Engine/Plugins/Media/MediaIOFramework/Source/MediaIOCore/Private/Player/MediaIOCorePlayerBase.cpp b/Engine/Plugins/Media/MediaIOFramework/Source/MediaIOCore/Private/Player/MediaIOCorePlayerBase.cpp
--- a/Engine/Plugins/Media/MediaIOFramework/Source/MediaIOCore/Private/Player/MediaIOCorePlayerBase.cpp
+++ b/Engine/Plugins/Media/MediaIOFramework/Source/MediaIOCore/Private/Player/MediaIOCorePlayerBase.cpp
@@ -140,16 +140,7 @@ bool FMediaIOCorePlayerBase::Open(const TSharedRef<FArchive, ESPMode::ThreadSafe
 
 void FMediaIOCorePlayerBase::TickTimeManagement()
 {
-	bool bUseTimecode = false;
 	if (bUseTimeSynchronization)
-	{
-		if (const UTimecodeProvider* Provider = GEngine->GetTimecodeProvider())
-		{
-			bUseTimecode = (Provider->GetSynchronizationState() == ETimecodeProviderSynchronizationState::Synchronized);
-		}
-	}
-
-	if (bUseTimecode)
 	{
 		FTimecode Timecode = FApp::GetTimecode();
 		FFrameRate FrameRate = FApp::GetTimecodeFrameRate();
diff --git a/Engine/Plugins/Media/MediaIOFramework/Source/MediaIOCore/Private/Shared/MediaIOCoreAudioSampleBase.cpp b/Engine/Plugins/Media/MediaIOFramework/Source/MediaIOCore/Private/Shared/MediaIOCoreAudioSampleBase.cpp
--- a/Engine/Plugins/Media/MediaIOFramework/Source/MediaIOCore/Private/Shared/MediaIOCoreAudioSampleBase.cpp
+++ b/Engine/Plugins/Media/MediaIOFramework/Source/MediaIOCore/Private/Shared/MediaIOCoreAudioSampleBase.cpp
@@ -83,7 +83,6 @@ bool FMediaIOCoreAudioSampleBase::SetProperties(uint32 InBufferSize, uint32 InNu
 void* FMediaIOCoreAudioSampleBase::RequestBuffer(uint32 InBufferSize)
 {
 	FreeSample();
-	Buffer.Reset(InBufferSize);
-	Buffer.SetNumUnsafeInternal(InBufferSize); // Reset the array without shrinking (Does not destruct items, does not de-allocate memory).
+	Buffer.SetNumUninitialized(InBufferSize); // Reset the array without shrinking (Does not destruct items, does not de-allocate memory).
 	return Buffer.GetData();
 }
diff --git a/Engine/Plugins/Media/MediaIOFramework/Source/MediaIOCore/Private/Shared/MediaIOCoreBinarySampleBase.cpp b/Engine/Plugins/Media/MediaIOFramework/Source/MediaIOCore/Private/Shared/MediaIOCoreBinarySampleBase.cpp
--- a/Engine/Plugins/Media/MediaIOFramework/Source/MediaIOCore/Private/Shared/MediaIOCoreBinarySampleBase.cpp
+++ b/Engine/Plugins/Media/MediaIOFramework/Source/MediaIOCore/Private/Shared/MediaIOCoreBinarySampleBase.cpp
@@ -73,7 +73,6 @@ bool FMediaIOCoreBinarySampleBase::SetProperties(FTimespan InTime, const FFrameR
 void* FMediaIOCoreBinarySampleBase::RequestBuffer(uint32 InBufferSize)
 {
 	FreeSample();
-	Buffer.Reset(InBufferSize);
-	Buffer.SetNumUnsafeInternal(InBufferSize); // Reset the array without shrinking (Does not destruct items, does not de-allocate memory).
+	Buffer.SetNumUninitialized(InBufferSize); // Reset the array without shrinking (Does not destruct items, does not de-allocate memory).
 	return Buffer.GetData();
 }
diff --git a/Engine/Plugins/Media/MediaIOFramework/Source/MediaIOCore/Private/Shared/MediaIOCoreTextureSampleBase.cpp b/Engine/Plugins/Media/MediaIOFramework/Source/MediaIOCore/Private/Shared/MediaIOCoreTextureSampleBase.cpp
--- a/Engine/Plugins/Media/MediaIOFramework/Source/MediaIOCore/Private/Shared/MediaIOCoreTextureSampleBase.cpp
+++ b/Engine/Plugins/Media/MediaIOFramework/Source/MediaIOCore/Private/Shared/MediaIOCoreTextureSampleBase.cpp
@@ -115,7 +115,6 @@ bool FMediaIOCoreTextureSampleBase::SetBufferWithEvenOddLine(bool bUseEvenLine,
 void* FMediaIOCoreTextureSampleBase::RequestBuffer(uint32 InBufferSize)
 {
 	FreeSample();
-	Buffer.Reset(InBufferSize);
-	Buffer.SetNumUnsafeInternal(InBufferSize); // Reset the array without shrinking (Does not destruct items, does not de-allocate memory).
+	Buffer.SetNumUninitialized(InBufferSize); // Reset the array without shrinking (Does not destruct items, does not de-allocate memory).
 	return Buffer.GetData();
 }
diff --git a/Engine/Plugins/Media/MediaIOFramework/Source/MediaIOCore/Public/MediaCapture.h b/Engine/Plugins/Media/MediaIOFramework/Source/MediaIOCore/Public/MediaCapture.h
--- a/Engine/Plugins/Media/MediaIOFramework/Source/MediaIOCore/Public/MediaCapture.h
+++ b/Engine/Plugins/Media/MediaIOFramework/Source/MediaIOCore/Public/MediaCapture.h
@@ -6,12 +6,13 @@
 #include "UObject/ObjectMacros.h"
 #include "UObject/Object.h"
 
+#include "HAL/CriticalSection.h"
 #include "MediaOutput.h"
 #include "Misc/Timecode.h"
 #include "PixelFormat.h"
 #include "RHI.h"
 #include "RHIResources.h"
-#include "HAL/CriticalSection.h"
+#include "Templates/Atomic.h"
 
 #include "MediaCapture.generated.h"
 
@@ -244,5 +245,5 @@ private:
 	FString MediaOutputName;
 
 	bool bResolvedTargetInitialized;
-	bool bWaitingForResolveCommandExecution;
+	TAtomic<int32> WaitingForResolveCommandExecutionCounter;
 };
diff --git a/Engine/Source/Runtime/Core/Private/Tests/Math/RangeSetTest.cpp b/Engine/Source/Runtime/Core/Private/Tests/Math/RangeSetTest.cpp
--- a/Engine/Source/Runtime/Core/Private/Tests/Math/RangeSetTest.cpp
+++ b/Engine/Source/Runtime/Core/Private/Tests/Math/RangeSetTest.cpp
@@ -9,26 +9,46 @@
 
 #if WITH_DEV_AUTOMATION_TESTS
 
-IMPLEMENT_SIMPLE_AUTOMATION_TEST(FRangeSetTest, "System.Core.Math.RangeSet", EAutomationTestFlags::Disabled | EAutomationTestFlags::EditorContext | EAutomationTestFlags::ClientContext | EAutomationTestFlags::SmokeFilter)
+IMPLEMENT_SIMPLE_AUTOMATION_TEST(FRangeSetTest, "System.Core.Math.RangeSet", EAutomationTestFlags::EditorContext | EAutomationTestFlags::ClientContext | EAutomationTestFlags::SmokeFilter)
 
 bool FRangeSetTest::RunTest(const FString& Parameters)
 {
-	// @todo: Write an actual unit test (the following just checks compilation of certain use cases)
-
 	{
-		TRange<int32> Range = TRange<int32>::AtLeast(0);
 		TRangeSet<int32> RangeSet;
+		RangeSet.Add(TRange<int32>::Inclusive(0, 1));
+		RangeSet.Add(TRange<int32>::Inclusive(1, 2));
+		RangeSet.Add(TRange<int32>::Inclusive(3, 4));
 
-		RangeSet.Add(Range);
 		int32 Value = RangeSet.GetMinBoundValue();
+		if (Value != 0)
+		{
+			return false;
+		}
+		
+		Value = RangeSet.GetMaxBoundValue();
+		if (Value != 4)
+		{
+			return false;
+		}
 	}
 
 	{
-		TRange<FTimespan> Range = TRange<FTimespan>::AtLeast(FTimespan::Zero());
 		TRangeSet<FTimespan> RangeSet;
+		RangeSet.Add(TRange<FTimespan>::Inclusive(FTimespan(0), FTimespan(1)));
+		RangeSet.Add(TRange<FTimespan>::Inclusive(FTimespan(1), FTimespan(2)));
+		RangeSet.Add(TRange<FTimespan>::Inclusive(FTimespan(3), FTimespan(4)));
 
-		RangeSet.Add(Range);
 		FTimespan Value = RangeSet.GetMinBoundValue();
+		if (Value != FTimespan::Zero())
+		{
+			return false;
+		}
+
+		Value = RangeSet.GetMaxBoundValue();
+		if (Value != FTimespan(4))
+		{
+			return false;
+		}
 	}
 
 	return true;
diff --git a/Engine/Source/Runtime/Core/Public/Math/RangeSet.h b/Engine/Source/Runtime/Core/Public/Math/RangeSet.h
--- a/Engine/Source/Runtime/Core/Public/Math/RangeSet.h
+++ b/Engine/Source/Runtime/Core/Public/Math/RangeSet.h
@@ -158,9 +158,13 @@ public:
 	{
 		BoundsType Result;
 
-		for (const auto& Range : Ranges)
+		if (Ranges.Num())
 		{
-			Result = BoundsType::MinLower(Result, Range.GetLowerBound());
+			Result = Ranges[0].GetLowerBound();
+			for (int32 i = 1; i < Ranges.Num(); i++)
+			{
+				Result = BoundsType::MinLower(Result, Ranges[i].GetLowerBound());
+			}
 		}
 
 		return Result;
@@ -189,9 +193,13 @@ public:
 	{
 		BoundsType Result;
 
-		for (const auto& Range : Ranges)
+		if (Ranges.Num())
 		{
-			Result = BoundsType::MaxUpper(Result, Range.GetUpperBound());
+			Result = Ranges[0].GetUpperBound();
+			for (int32 i = 1; i < Ranges.Num(); i++)
+			{
+				Result = BoundsType::MaxUpper(Result, Ranges[i].GetUpperBound());
+			}
 		}
 
 		return Result;
