[AJA]
Correct the FPS for Dual Link
Try to select the item from the next column even if the previous didn't match

diff --git a/Engine/Plugins/Media/MediaIOFramework/Source/MediaIOEditor/Public/Widgets/SMediaPermutationsSelector.inl b/Engine/Plugins/Media/MediaIOFramework/Source/MediaIOEditor/Public/Widgets/SMediaPermutationsSelector.inl
--- a/Engine/Plugins/Media/MediaIOFramework/Source/MediaIOEditor/Public/Widgets/SMediaPermutationsSelector.inl
+++ b/Engine/Plugins/Media/MediaIOFramework/Source/MediaIOEditor/Public/Widgets/SMediaPermutationsSelector.inl
@@ -241,6 +241,8 @@ void SMediaPermutationsSelector<ItemType, ItemBuilder>::ItemSelected(int32 Uniqu
 	// Try to find something that matches what we used to have.
 	for (int32 Index = MaxColumns; Index < PropertyColumns.Num(); ++Index)
 	{
+		TArray<int32> CopiedAllValidItemIndexes = AllValidItemIndexes;
+
 		const FColumn& Column = PropertyColumns[Index];
 		for (int32 ItemIndex = AllValidItemIndexes.Num() - 1; ItemIndex >= 0; --ItemIndex)
 		{
@@ -250,15 +252,35 @@ void SMediaPermutationsSelector<ItemType, ItemBuilder>::ItemSelected(int32 Uniqu
 			}
 		}
 	
-		if (AllValidItemIndexes.Num() > 0)
+		if (AllValidItemIndexes.Num() == 0)
 		{
-			SelectedPermutationIndex = AllValidItemIndexes[0];
+			if ((Index + 1) < PropertyColumns.Num())
+			{
+				// if there are no options, go to the next in the list
+				AllValidItemIndexes = CopiedAllValidItemIndexes;
+			}
+			else if (CopiedAllValidItemIndexes.Num() > 0)
+			{
+				SelectedPermutationIndex = CopiedAllValidItemIndexes[0];
+			}
 		}
 	}
-	
-	// There should always be one selected item
-	if (SelectedPermutationIndex == INDEX_NONE)
+
+	if (AllValidItemIndexes.Num())
+	{
+		// Get the first item. (RemoveAtSwap change the order.)
+		SelectedPermutationIndex = AllValidItemIndexes[0];
+		for (int32 IndexValue : AllValidItemIndexes)
+		{
+			if (IndexValue < SelectedPermutationIndex)
+			{
+				SelectedPermutationIndex = IndexValue;
+			}
+		}
+	}
+	else if (SelectedPermutationIndex == INDEX_NONE)
 	{
+		// There should always be one selected item
 		SelectedPermutationIndex = UniqueItemIndex;
 	}
 	
