Update media capture to use ENQUEUE_RENDER_COMMAND

diff --git a/Engine/Plugins/Media/MediaIOFramework/Source/MediaIOCore/Private/MediaCapture.cpp b/Engine/Plugins/Media/MediaIOFramework/Source/MediaIOCore/Private/MediaCapture.cpp
--- a/Engine/Plugins/Media/MediaIOFramework/Source/MediaIOCore/Private/MediaCapture.cpp
+++ b/Engine/Plugins/Media/MediaIOFramework/Source/MediaIOCore/Private/MediaCapture.cpp
@@ -481,243 +481,239 @@ void UMediaCapture::OnEndFrame_GameThread()
 		CapturingFrame->UserData = GetCaptureFrameUserData_GameThread();
 	}
 
-	// RenderCommand to be executed on the RenderThread
-	auto RenderCommand = [this](FRHICommandListImmediate& RHICmdList, FCaptureFrame* InCapturingFrame, FCaptureFrame* InReadyFrame, TWeakPtr<FSceneViewport> InViewport, FTextureRenderTargetResource* InTextureRenderTargetResource)
+	++WaitingForResolveCommandExecutionCounter;
+
+	// Init variables for ENQUEUE_RENDER_COMMAND
+	TWeakPtr<FSceneViewport> InCapturingSceneViewport = CapturingSceneViewport;
+	FTextureRenderTargetResource* InTextureRenderTargetResource = nullptr;
+	FIntPoint InDesiredSize = DesiredSize;
+	UMediaCapture* InMediaCapture = this;
+
 	{
-		FTexture2DRHIRef SourceTexture;
+		FScopeLock Lock(&AccessingCapturingSource);
+		if (CapturingRenderTarget)
 		{
-			TSharedPtr<FSceneViewport> SceneViewportPtr = InViewport.Pin();
-			if (SceneViewportPtr)
+			InTextureRenderTargetResource = CapturingRenderTarget->GameThread_GetRenderTargetResource();
+		}
+	}
+
+	// RenderCommand to be executed on the RenderThread
+	ENQUEUE_RENDER_COMMAND(FMediaOutputCaptureFrameCreateTexture)(
+		[CapturingFrame, ReadyFrame, InCapturingSceneViewport, InTextureRenderTargetResource, InDesiredSize, InMediaCapture](FRHICommandListImmediate& RHICmdList)
+		{
+			FTexture2DRHIRef SourceTexture;
 			{
-				SourceTexture = SceneViewportPtr->GetRenderTargetTexture();
-				if (!SourceTexture.IsValid() && SceneViewportPtr->GetViewportRHI())
+				TSharedPtr<FSceneViewport> SceneViewportPtr = InCapturingSceneViewport.Pin();
+				if (SceneViewportPtr)
+				{
+					SourceTexture = SceneViewportPtr->GetRenderTargetTexture();
+					if (!SourceTexture.IsValid() && SceneViewportPtr->GetViewportRHI())
+					{
+						SourceTexture = RHICmdList.GetViewportBackBuffer(SceneViewportPtr->GetViewportRHI());
+					}
+				}
+				else if (InTextureRenderTargetResource && InTextureRenderTargetResource->GetTextureRenderTarget2DResource())
 				{
-					SourceTexture = RHICmdList.GetViewportBackBuffer(SceneViewportPtr->GetViewportRHI());
+					SourceTexture = InTextureRenderTargetResource->GetTextureRenderTarget2DResource()->GetTextureRHI();
 				}
 			}
-			else if (InTextureRenderTargetResource && InTextureRenderTargetResource->GetTextureRenderTarget2DResource())
-			{
-				SourceTexture = InTextureRenderTargetResource->GetTextureRenderTarget2DResource()->GetTextureRHI();
-			}
-		}
 
-		if (!SourceTexture.IsValid())
-		{
-			MediaState = EMediaCaptureState::Error;
-			UE_LOG(LogMediaIOCore, Error, TEXT("Can't grab the Texture to capture for '%s'."), *MediaOutputName);
-		}
-		else if (InCapturingFrame)
-		{
-			if (DesiredPixelFormat != SourceTexture->GetFormat())
+			if (!SourceTexture.IsValid())
 			{
-				MediaState = EMediaCaptureState::Error;
-				UE_LOG(LogMediaIOCore, Error, TEXT("The capture will stop for '%s'. The Source pixel format doesn't match with the user requested pixel format. Requested: %s Source: %s")
-					, *MediaOutputName
-					, GetPixelFormatString(DesiredPixelFormat)
-					, GetPixelFormatString(SourceTexture->GetFormat()));
+				InMediaCapture->MediaState = EMediaCaptureState::Error;
+				UE_LOG(LogMediaIOCore, Error, TEXT("Can't grab the Texture to capture for '%s'."), *InMediaCapture->MediaOutputName);
 			}
-			else if (DesiredCaptureOptions.Crop == EMediaCaptureCroppingType::None)
+			else if (CapturingFrame)
 			{
-				if (DesiredSize.X != SourceTexture->GetSizeX() || DesiredSize.Y != SourceTexture->GetSizeY())
+				if (InMediaCapture->DesiredPixelFormat != SourceTexture->GetFormat())
 				{
-					MediaState = EMediaCaptureState::Error;
-					UE_LOG(LogMediaIOCore, Error, TEXT("The capture will stop for '%s'. The Source size doesn't match with the user requested size. Requested: %d,%d  Source: %d,%d")
-						, *MediaOutputName
-						, DesiredSize.X, DesiredSize.Y
-						, SourceTexture->GetSizeX(), SourceTexture->GetSizeY());
+					InMediaCapture->MediaState = EMediaCaptureState::Error;
+					UE_LOG(LogMediaIOCore, Error, TEXT("The capture will stop for '%s'. The Source pixel format doesn't match with the user requested pixel format. Requested: %s Source: %s")
+						, *InMediaCapture->MediaOutputName
+						, GetPixelFormatString(InMediaCapture->DesiredPixelFormat)
+						, GetPixelFormatString(SourceTexture->GetFormat()));
 				}
-			}
-			else
-			{
-				FIntPoint StartCapturePoint = FIntPoint::ZeroValue;
-				if (DesiredCaptureOptions.Crop == EMediaCaptureCroppingType::Custom)
+				else if (InMediaCapture->DesiredCaptureOptions.Crop == EMediaCaptureCroppingType::None)
 				{
-					StartCapturePoint = DesiredCaptureOptions.CustomCapturePoint;
+					if (InDesiredSize.X != SourceTexture->GetSizeX() || InDesiredSize.Y != SourceTexture->GetSizeY())
+					{
+						InMediaCapture->MediaState = EMediaCaptureState::Error;
+						UE_LOG(LogMediaIOCore, Error, TEXT("The capture will stop for '%s'. The Source size doesn't match with the user requested size. Requested: %d,%d  Source: %d,%d")
+							, *InMediaCapture->MediaOutputName
+							, InDesiredSize.X, InDesiredSize.Y
+							, SourceTexture->GetSizeX(), SourceTexture->GetSizeY());
+					}
 				}
-
-				if ((uint32)(DesiredSize.X + StartCapturePoint.X) > SourceTexture->GetSizeX() || (uint32)(DesiredSize.Y + StartCapturePoint.Y) > SourceTexture->GetSizeY())
+				else
 				{
-					MediaState = EMediaCaptureState::Error;
-					UE_LOG(LogMediaIOCore, Error, TEXT("The capture will stop for '%s'. The Source size doesn't match with the user requested size. Requested: %d,%d  Source: %d,%d")
-						, *MediaOutputName
-						, DesiredSize.X, DesiredSize.Y
-						, SourceTexture->GetSizeX(), SourceTexture->GetSizeY());
+					FIntPoint StartCapturePoint = FIntPoint::ZeroValue;
+					if (InMediaCapture->DesiredCaptureOptions.Crop == EMediaCaptureCroppingType::Custom)
+					{
+						StartCapturePoint = InMediaCapture->DesiredCaptureOptions.CustomCapturePoint;
+					}
+
+					if ((uint32)(InDesiredSize.X + StartCapturePoint.X) > SourceTexture->GetSizeX() || (uint32)(InDesiredSize.Y + StartCapturePoint.Y) > SourceTexture->GetSizeY())
+					{
+						InMediaCapture->MediaState = EMediaCaptureState::Error;
+						UE_LOG(LogMediaIOCore, Error, TEXT("The capture will stop for '%s'. The Source size doesn't match with the user requested size. Requested: %d,%d  Source: %d,%d")
+							, *InMediaCapture->MediaOutputName
+							, InDesiredSize.X, InDesiredSize.Y
+							, SourceTexture->GetSizeX(), SourceTexture->GetSizeY());
+					}
 				}
 			}
-		}
 
-		if (InCapturingFrame && MediaState != EMediaCaptureState::Error)
-		{
-			SCOPE_CYCLE_COUNTER(STAT_MediaCapture_RenderThread_CopyToResolve);
-
-			FPooledRenderTargetDesc OutputDesc = FPooledRenderTargetDesc::Create2DDesc(
-				DesiredOutputSize,
-				DesiredOutputPixelFormat,
-				FClearValueBinding::None,
-				TexCreate_None,
-				TexCreate_RenderTargetable,
-				false);
-			TRefCountPtr<IPooledRenderTarget> ResampleTexturePooledRenderTarget;
-			GetRendererModule().RenderTargetPoolFindFreeElement(RHICmdList, OutputDesc, ResampleTexturePooledRenderTarget, TEXT("MediaCapture"));
-			const FSceneRenderTargetItem& DestRenderTarget = ResampleTexturePooledRenderTarget->GetRenderTargetItem();
-
-			// Do we need to crop
-			float ULeft = 0.0f;
-			float URight = 1.0f;
-			float VTop = 0.0f;
-			float VBottom = 1.0f;
-			FResolveParams ResolveParams;
-			if (DesiredCaptureOptions.Crop != EMediaCaptureCroppingType::None)
+			if (CapturingFrame && InMediaCapture->MediaState != EMediaCaptureState::Error)
 			{
-				switch (DesiredCaptureOptions.Crop)
+				SCOPE_CYCLE_COUNTER(STAT_MediaCapture_RenderThread_CopyToResolve);
+
+				FPooledRenderTargetDesc OutputDesc = FPooledRenderTargetDesc::Create2DDesc(
+					InMediaCapture->DesiredOutputSize,
+					InMediaCapture->DesiredOutputPixelFormat,
+					FClearValueBinding::None,
+					TexCreate_None,
+					TexCreate_RenderTargetable,
+					false);
+				TRefCountPtr<IPooledRenderTarget> ResampleTexturePooledRenderTarget;
+				GetRendererModule().RenderTargetPoolFindFreeElement(RHICmdList, OutputDesc, ResampleTexturePooledRenderTarget, TEXT("MediaCapture"));
+				const FSceneRenderTargetItem& DestRenderTarget = ResampleTexturePooledRenderTarget->GetRenderTargetItem();
+
+				// Do we need to crop
+				float ULeft = 0.0f;
+				float URight = 1.0f;
+				float VTop = 0.0f;
+				float VBottom = 1.0f;
+				FResolveParams ResolveParams;
+				if (InMediaCapture->DesiredCaptureOptions.Crop != EMediaCaptureCroppingType::None)
 				{
-				case EMediaCaptureCroppingType::Center:
-					ResolveParams.Rect = FResolveRect((SourceTexture->GetSizeX() - DesiredSize.X) / 2, (SourceTexture->GetSizeY() - DesiredSize.Y) / 2, 0, 0);
-					ResolveParams.Rect.X2 = ResolveParams.Rect.X1 + DesiredSize.X;
-					ResolveParams.Rect.Y2 = ResolveParams.Rect.Y1 + DesiredSize.Y;
-					break;
-				case EMediaCaptureCroppingType::TopLeft:
-					ResolveParams.Rect = FResolveRect(0, 0, DesiredSize.X, DesiredSize.Y);
-					break;
-				case EMediaCaptureCroppingType::Custom:
-					ResolveParams.Rect = FResolveRect(DesiredCaptureOptions.CustomCapturePoint.X, DesiredCaptureOptions.CustomCapturePoint.Y, 0, 0);
-					ResolveParams.Rect.X2 = ResolveParams.Rect.X1 + DesiredSize.X;
-					ResolveParams.Rect.Y2 = ResolveParams.Rect.Y1 + DesiredSize.Y;
-					break;
-				}
-
-				ResolveParams.DestRect.X1 = 0;
-				ResolveParams.DestRect.X2 = DesiredSize.X;
-				ResolveParams.DestRect.Y1 = 0;
-				ResolveParams.DestRect.Y2 = DesiredSize.Y;
-
-				ULeft = (float)ResolveParams.Rect.X1 / (float)SourceTexture->GetSizeX();
-				URight = (float)ResolveParams.Rect.X2 / (float)SourceTexture->GetSizeX();
-				VTop = (float)ResolveParams.Rect.Y1 / (float)SourceTexture->GetSizeY();
-				VBottom = (float)ResolveParams.Rect.Y2 / (float)SourceTexture->GetSizeY();
-			}
+					switch (InMediaCapture->DesiredCaptureOptions.Crop)
+					{
+					case EMediaCaptureCroppingType::Center:
+						ResolveParams.Rect = FResolveRect((SourceTexture->GetSizeX() - InDesiredSize.X) / 2, (SourceTexture->GetSizeY() - InDesiredSize.Y) / 2, 0, 0);
+						ResolveParams.Rect.X2 = ResolveParams.Rect.X1 + InDesiredSize.X;
+						ResolveParams.Rect.Y2 = ResolveParams.Rect.Y1 + InDesiredSize.Y;
+						break;
+					case EMediaCaptureCroppingType::TopLeft:
+						ResolveParams.Rect = FResolveRect(0, 0, InDesiredSize.X, InDesiredSize.Y);
+						break;
+					case EMediaCaptureCroppingType::Custom:
+						ResolveParams.Rect = FResolveRect(InMediaCapture->DesiredCaptureOptions.CustomCapturePoint.X, InMediaCapture->DesiredCaptureOptions.CustomCapturePoint.Y, 0, 0);
+						ResolveParams.Rect.X2 = ResolveParams.Rect.X1 + InDesiredSize.X;
+						ResolveParams.Rect.Y2 = ResolveParams.Rect.Y1 + InDesiredSize.Y;
+						break;
+					}
 
-			{
-				SCOPED_DRAW_EVENTF(RHICmdList, MediaCapture, TEXT("MediaCapture"));
+					ResolveParams.DestRect.X1 = 0;
+					ResolveParams.DestRect.X2 = InDesiredSize.X;
+					ResolveParams.DestRect.Y1 = 0;
+					ResolveParams.DestRect.Y2 = InDesiredSize.Y;
 
-				if (ConversionOperation == EMediaCaptureConversionOperation::NONE)
-				{
-					// Asynchronously copy target from GPU to GPU
-					RHICmdList.CopyToResolveTarget(SourceTexture, DestRenderTarget.TargetableTexture, ResolveParams);
+					ULeft = (float)ResolveParams.Rect.X1 / (float)SourceTexture->GetSizeX();
+					URight = (float)ResolveParams.Rect.X2 / (float)SourceTexture->GetSizeX();
+					VTop = (float)ResolveParams.Rect.Y1 / (float)SourceTexture->GetSizeY();
+					VBottom = (float)ResolveParams.Rect.Y2 / (float)SourceTexture->GetSizeY();
 				}
-				else
+
 				{
-					// convert the source with a draw call
-					FGraphicsPipelineStateInitializer GraphicsPSOInit;
-					FRHITexture* RenderTarget = DestRenderTarget.TargetableTexture.GetReference();
-					SetRenderTargets(RHICmdList, 1, &RenderTarget, nullptr, ESimpleRenderTargetMode::EExistingColorAndDepth, FExclusiveDepthStencil::DepthNop_StencilNop);
+					SCOPED_DRAW_EVENTF(RHICmdList, MediaCapture, TEXT("MediaCapture"));
 
-					RHICmdList.ApplyCachedRenderTargets(GraphicsPSOInit);
+					if (InMediaCapture->ConversionOperation == EMediaCaptureConversionOperation::NONE)
+					{
+						// Asynchronously copy target from GPU to GPU
+						RHICmdList.CopyToResolveTarget(SourceTexture, DestRenderTarget.TargetableTexture, ResolveParams);
+					}
+					else
+					{
+						// convert the source with a draw call
+						FGraphicsPipelineStateInitializer GraphicsPSOInit;
+						FRHITexture* RenderTarget = DestRenderTarget.TargetableTexture.GetReference();
+						SetRenderTargets(RHICmdList, 1, &RenderTarget, nullptr, ESimpleRenderTargetMode::EExistingColorAndDepth, FExclusiveDepthStencil::DepthNop_StencilNop);
 
-					GraphicsPSOInit.DepthStencilState = TStaticDepthStencilState<false, CF_Always>::GetRHI();
-					GraphicsPSOInit.RasterizerState = TStaticRasterizerState<>::GetRHI();
-					GraphicsPSOInit.BlendState = TStaticBlendStateWriteMask<CW_RGBA, CW_NONE, CW_NONE, CW_NONE, CW_NONE, CW_NONE, CW_NONE, CW_NONE>::GetRHI();
-					GraphicsPSOInit.PrimitiveType = PT_TriangleStrip;
+						RHICmdList.ApplyCachedRenderTargets(GraphicsPSOInit);
 
-					// configure media shaders
-					auto ShaderMap = GetGlobalShaderMap(GMaxRHIFeatureLevel);
-					TShaderMapRef<FMediaShadersVS> VertexShader(ShaderMap);
+						GraphicsPSOInit.DepthStencilState = TStaticDepthStencilState<false, CF_Always>::GetRHI();
+						GraphicsPSOInit.RasterizerState = TStaticRasterizerState<>::GetRHI();
+						GraphicsPSOInit.BlendState = TStaticBlendStateWriteMask<CW_RGBA, CW_NONE, CW_NONE, CW_NONE, CW_NONE, CW_NONE, CW_NONE, CW_NONE>::GetRHI();
+						GraphicsPSOInit.PrimitiveType = PT_TriangleStrip;
 
-					GraphicsPSOInit.BoundShaderState.VertexDeclarationRHI = GMediaVertexDeclaration.VertexDeclarationRHI;
-					GraphicsPSOInit.BoundShaderState.VertexShaderRHI = GETSAFERHISHADER_VERTEX(*VertexShader);
+						// configure media shaders
+						auto ShaderMap = GetGlobalShaderMap(GMaxRHIFeatureLevel);
+						TShaderMapRef<FMediaShadersVS> VertexShader(ShaderMap);
 
-					const bool bDoLinearToSRGB = false;
+						GraphicsPSOInit.BoundShaderState.VertexDeclarationRHI = GMediaVertexDeclaration.VertexDeclarationRHI;
+						GraphicsPSOInit.BoundShaderState.VertexShaderRHI = GETSAFERHISHADER_VERTEX(*VertexShader);
 
-					switch (ConversionOperation)
-					{
-					case EMediaCaptureConversionOperation::RGBA8_TO_YUV_8BIT:
-						{
-							TShaderMapRef<FRGB8toUYVY8ConvertPS> ConvertShader(ShaderMap);
-							GraphicsPSOInit.BoundShaderState.PixelShaderRHI = GETSAFERHISHADER_PIXEL(*ConvertShader);
-							SetGraphicsPipelineState(RHICmdList, GraphicsPSOInit);
-						ConvertShader->SetParameters(RHICmdList, SourceTexture, MediaShaders::RgbToYuvRec709Full, MediaShaders::YUVOffset8bits, bDoLinearToSRGB);
-						}
-						break;
-					case EMediaCaptureConversionOperation::RGB10_TO_YUVv210_10BIT:
-						{
-							TShaderMapRef<FRGB10toYUVv210ConvertPS> ConvertShader(ShaderMap);
-							GraphicsPSOInit.BoundShaderState.PixelShaderRHI = GETSAFERHISHADER_PIXEL(*ConvertShader);
-							SetGraphicsPipelineState(RHICmdList, GraphicsPSOInit);
-						ConvertShader->SetParameters(RHICmdList, SourceTexture, MediaShaders::RgbToYuvRec709Full, MediaShaders::YUVOffset10bits, bDoLinearToSRGB);
-						}
-						break;
-					case EMediaCaptureConversionOperation::INVERT_ALPHA:
+						const bool bDoLinearToSRGB = false;
+
+						switch (InMediaCapture->ConversionOperation)
 						{
-							TShaderMapRef<FInvertAlphaPS> ConvertShader(ShaderMap);
-							GraphicsPSOInit.BoundShaderState.PixelShaderRHI = GETSAFERHISHADER_PIXEL(*ConvertShader);
-							SetGraphicsPipelineState(RHICmdList, GraphicsPSOInit);
-							ConvertShader->SetParameters(RHICmdList, SourceTexture);
+						case EMediaCaptureConversionOperation::RGBA8_TO_YUV_8BIT:
+							{
+								TShaderMapRef<FRGB8toUYVY8ConvertPS> ConvertShader(ShaderMap);
+								GraphicsPSOInit.BoundShaderState.PixelShaderRHI = GETSAFERHISHADER_PIXEL(*ConvertShader);
+								SetGraphicsPipelineState(RHICmdList, GraphicsPSOInit);
+							ConvertShader->SetParameters(RHICmdList, SourceTexture, MediaShaders::RgbToYuvRec709Full, MediaShaders::YUVOffset8bits, bDoLinearToSRGB);
+							}
+							break;
+						case EMediaCaptureConversionOperation::RGB10_TO_YUVv210_10BIT:
+							{
+								TShaderMapRef<FRGB10toYUVv210ConvertPS> ConvertShader(ShaderMap);
+								GraphicsPSOInit.BoundShaderState.PixelShaderRHI = GETSAFERHISHADER_PIXEL(*ConvertShader);
+								SetGraphicsPipelineState(RHICmdList, GraphicsPSOInit);
+							ConvertShader->SetParameters(RHICmdList, SourceTexture, MediaShaders::RgbToYuvRec709Full, MediaShaders::YUVOffset10bits, bDoLinearToSRGB);
+							}
+							break;
+						case EMediaCaptureConversionOperation::INVERT_ALPHA:
+							{
+								TShaderMapRef<FInvertAlphaPS> ConvertShader(ShaderMap);
+								GraphicsPSOInit.BoundShaderState.PixelShaderRHI = GETSAFERHISHADER_PIXEL(*ConvertShader);
+								SetGraphicsPipelineState(RHICmdList, GraphicsPSOInit);
+								ConvertShader->SetParameters(RHICmdList, SourceTexture);
+							}
+							break;
 						}
-						break;
-					}
 
-					// draw full size quad into render target
-					FVertexBufferRHIRef VertexBuffer = CreateTempMediaVertexBuffer(ULeft, URight, VTop, VBottom);
-					RHICmdList.SetStreamSource(0, VertexBuffer, 0);
+						// draw full size quad into render target
+						FVertexBufferRHIRef VertexBuffer = CreateTempMediaVertexBuffer(ULeft, URight, VTop, VBottom);
+						RHICmdList.SetStreamSource(0, VertexBuffer, 0);
 
-					// set viewport to RT size
-					RHICmdList.SetViewport(0, 0, 0.0f, DesiredOutputSize.X, DesiredOutputSize.Y, 1.0f);
-					RHICmdList.DrawPrimitive(0, 2, 1);
-					RHICmdList.TransitionResource(EResourceTransitionAccess::EReadable, DestRenderTarget.TargetableTexture);
+						// set viewport to RT size
+						RHICmdList.SetViewport(0, 0, 0.0f, InMediaCapture->DesiredOutputSize.X, InMediaCapture->DesiredOutputSize.Y, 1.0f);
+						RHICmdList.DrawPrimitive(0, 2, 1);
+						RHICmdList.TransitionResource(EResourceTransitionAccess::EReadable, DestRenderTarget.TargetableTexture);
+					}
 				}
-			}
 
-			// Asynchronously copy duplicate target from GPU to System Memory
-			RHICmdList.CopyToResolveTarget(DestRenderTarget.TargetableTexture, InCapturingFrame->ReadbackTexture, FResolveParams());
-
-			InCapturingFrame->bResolvedTargetRequested = true;
-		}
+				// Asynchronously copy duplicate target from GPU to System Memory
+				RHICmdList.CopyToResolveTarget(DestRenderTarget.TargetableTexture, CapturingFrame->ReadbackTexture, FResolveParams());
 
-		if (InReadyFrame && MediaState != EMediaCaptureState::Error)
-		{
-			check(InReadyFrame->ReadbackTexture.IsValid());
-
-			// Lock & read
-			void* ColorDataBuffer = nullptr;
-			int32 Width = 0, Height = 0;
-			{
-				SCOPE_CYCLE_COUNTER(STAT_MediaCapture_RenderThread_MapStaging);
-				RHICmdList.MapStagingSurface(InReadyFrame->ReadbackTexture, ColorDataBuffer, Width, Height);
+				CapturingFrame->bResolvedTargetRequested = true;
 			}
 
+			if (ReadyFrame && InMediaCapture->MediaState != EMediaCaptureState::Error)
 			{
-				SCOPE_CYCLE_COUNTER(STAT_MediaCapture_RenderThread_Callback);
-				OnFrameCaptured_RenderingThread(InReadyFrame->CaptureBaseData, InReadyFrame->UserData, ColorDataBuffer, Width, Height);
-			}
-			InReadyFrame->bResolvedTargetRequested = false;
+				check(ReadyFrame->ReadbackTexture.IsValid());
 
-			RHICmdList.UnmapStagingSurface(InReadyFrame->ReadbackTexture);
-		}
+				// Lock & read
+				void* ColorDataBuffer = nullptr;
+				int32 Width = 0, Height = 0;
+				{
+					SCOPE_CYCLE_COUNTER(STAT_MediaCapture_RenderThread_MapStaging);
+					RHICmdList.MapStagingSurface(ReadyFrame->ReadbackTexture, ColorDataBuffer, Width, Height);
+				}
 
-		--WaitingForResolveCommandExecutionCounter;
-	};
+				{
+					SCOPE_CYCLE_COUNTER(STAT_MediaCapture_RenderThread_Callback);
+					InMediaCapture->OnFrameCaptured_RenderingThread(ReadyFrame->CaptureBaseData, ReadyFrame->UserData, ColorDataBuffer, Width, Height);
+				}
+				ReadyFrame->bResolvedTargetRequested = false;
 
-	FTextureRenderTargetResource* TextureRenderTargetResource = nullptr;
-	{
-		FScopeLock Lock(&AccessingCapturingSource);
+				RHICmdList.UnmapStagingSurface(ReadyFrame->ReadbackTexture);
+			}
 
-		if(CapturingRenderTarget)
-		{
-			TextureRenderTargetResource = CapturingRenderTarget->GameThread_GetRenderTargetResource();
+			--InMediaCapture->WaitingForResolveCommandExecutionCounter;
 		}
-	}
-
-	++WaitingForResolveCommandExecutionCounter;
-	ENQUEUE_UNIQUE_RENDER_COMMAND_FIVEPARAMETER(
-		MediaOutputCaptureFrameCreateTexture,
-		FCaptureFrame*, InCapturingFrame, CapturingFrame,
-		FCaptureFrame*, InPreviousFrame, ReadyFrame,
-		TWeakPtr<FSceneViewport>, InViewport, CapturingSceneViewport,
-		FTextureRenderTargetResource*, InTextureRenderTargetResource, TextureRenderTargetResource,
-		decltype(RenderCommand), InRenderCommand, RenderCommand,
-		{
-			InRenderCommand(RHICmdList, InCapturingFrame, InPreviousFrame, InViewport, InTextureRenderTargetResource);
-		});
+	);
 }
 
 /* namespace MediaCaptureDetails implementation
