MediaCapture capture Slate in standalone and in packaged build.

diff --git a/Engine/Plugins/Media/MediaFrameworkUtilities/Source/MediaFrameworkUtilitiesEditor/Private/CaptureTab/SMediaFrameworkCaptureOutputWidget.cpp b/Engine/Plugins/Media/MediaFrameworkUtilities/Source/MediaFrameworkUtilitiesEditor/Private/CaptureTab/SMediaFrameworkCaptureOutputWidget.cpp
--- a/Engine/Plugins/Media/MediaFrameworkUtilities/Source/MediaFrameworkUtilitiesEditor/Private/CaptureTab/SMediaFrameworkCaptureOutputWidget.cpp
+++ b/Engine/Plugins/Media/MediaFrameworkUtilities/Source/MediaFrameworkUtilitiesEditor/Private/CaptureTab/SMediaFrameworkCaptureOutputWidget.cpp
@@ -678,7 +678,7 @@ void SMediaFrameworkCaptureCurrentViewportWidget::OnPrePIEEnded()
 void SMediaFrameworkCaptureCurrentViewportWidget::ShutdownViewport()
 {
 	TSharedPtr<FSceneViewport> EditorSceneViewportPin = EditorSceneViewport.Pin();
-	if (EditorSceneViewportPin.IsValid())
+	if (EditorSceneViewportPin.IsValid() && EditorSceneViewportPin->GetViewportWidget().IsValid())
 	{
 		EditorSceneViewportPin->SetFixedViewportSize(0, 0);
 	}
diff --git a/Engine/Plugins/Media/MediaIOFramework/Source/MediaIOCore/Private/MediaCapture.cpp b/Engine/Plugins/Media/MediaIOFramework/Source/MediaIOCore/Private/MediaCapture.cpp
--- a/Engine/Plugins/Media/MediaIOFramework/Source/MediaIOCore/Private/MediaCapture.cpp
+++ b/Engine/Plugins/Media/MediaIOFramework/Source/MediaIOCore/Private/MediaCapture.cpp
@@ -500,216 +500,228 @@ void UMediaCapture::OnEndFrame_GameThread()
 	// RenderCommand to be executed on the RenderThread
 	ENQUEUE_RENDER_COMMAND(FMediaOutputCaptureFrameCreateTexture)(
 		[CapturingFrame, ReadyFrame, InCapturingSceneViewport, InTextureRenderTargetResource, InDesiredSize, InMediaCapture](FRHICommandListImmediate& RHICmdList)
+	{
+		FTexture2DRHIRef SourceTexture;
 		{
-			FTexture2DRHIRef SourceTexture;
+			TSharedPtr<FSceneViewport> SceneViewportPtr = InCapturingSceneViewport.Pin();
+			if (SceneViewportPtr)
 			{
-				TSharedPtr<FSceneViewport> SceneViewportPtr = InCapturingSceneViewport.Pin();
-				if (SceneViewportPtr)
+#if WITH_EDITOR
+				if (!IsRunningGame())
 				{
+					// PIE, PIE in windows, editor viewport
 					SourceTexture = SceneViewportPtr->GetRenderTargetTexture();
 					if (!SourceTexture.IsValid() && SceneViewportPtr->GetViewportRHI())
 					{
 						SourceTexture = RHICmdList.GetViewportBackBuffer(SceneViewportPtr->GetViewportRHI());
 					}
 				}
-				else if (InTextureRenderTargetResource && InTextureRenderTargetResource->GetTextureRenderTarget2DResource())
+				else
+#endif
+				if (SceneViewportPtr->GetViewportRHI())
 				{
-					SourceTexture = InTextureRenderTargetResource->GetTextureRenderTarget2DResource()->GetTextureRHI();
+					// Standalone and packaged
+					SourceTexture = RHICmdList.GetViewportBackBuffer(SceneViewportPtr->GetViewportRHI());
 				}
 			}
-
-			if (!SourceTexture.IsValid())
+			else if (InTextureRenderTargetResource && InTextureRenderTargetResource->GetTextureRenderTarget2DResource())
 			{
+				SourceTexture = InTextureRenderTargetResource->GetTextureRenderTarget2DResource()->GetTextureRHI();
+			}
+		}
+
+		if (!SourceTexture.IsValid())
+		{
 				InMediaCapture->MediaState = EMediaCaptureState::Error;
 				UE_LOG(LogMediaIOCore, Error, TEXT("Can't grab the Texture to capture for '%s'."), *InMediaCapture->MediaOutputName);
-			}
+		}
 			else if (CapturingFrame)
-			{
+		{
 				if (InMediaCapture->DesiredPixelFormat != SourceTexture->GetFormat())
-				{
+			{
 					InMediaCapture->MediaState = EMediaCaptureState::Error;
-					UE_LOG(LogMediaIOCore, Error, TEXT("The capture will stop for '%s'. The Source pixel format doesn't match with the user requested pixel format. Requested: %s Source: %s")
+				UE_LOG(LogMediaIOCore, Error, TEXT("The capture will stop for '%s'. The Source pixel format doesn't match with the user requested pixel format. Requested: %s Source: %s")
 						, *InMediaCapture->MediaOutputName
 						, GetPixelFormatString(InMediaCapture->DesiredPixelFormat)
-						, GetPixelFormatString(SourceTexture->GetFormat()));
-				}
+					, GetPixelFormatString(SourceTexture->GetFormat()));
+			}
 				else if (InMediaCapture->DesiredCaptureOptions.Crop == EMediaCaptureCroppingType::None)
-				{
+			{
 					if (InDesiredSize.X != SourceTexture->GetSizeX() || InDesiredSize.Y != SourceTexture->GetSizeY())
-					{
+				{
 						InMediaCapture->MediaState = EMediaCaptureState::Error;
-						UE_LOG(LogMediaIOCore, Error, TEXT("The capture will stop for '%s'. The Source size doesn't match with the user requested size. Requested: %d,%d  Source: %d,%d")
+					UE_LOG(LogMediaIOCore, Error, TEXT("The capture will stop for '%s'. The Source size doesn't match with the user requested size. Requested: %d,%d  Source: %d,%d")
 							, *InMediaCapture->MediaOutputName
 							, InDesiredSize.X, InDesiredSize.Y
-							, SourceTexture->GetSizeX(), SourceTexture->GetSizeY());
-					}
+						, SourceTexture->GetSizeX(), SourceTexture->GetSizeY());
 				}
-				else
-				{
-					FIntPoint StartCapturePoint = FIntPoint::ZeroValue;
+			}
+			else
+			{
+				FIntPoint StartCapturePoint = FIntPoint::ZeroValue;
 					if (InMediaCapture->DesiredCaptureOptions.Crop == EMediaCaptureCroppingType::Custom)
-					{
+				{
 						StartCapturePoint = InMediaCapture->DesiredCaptureOptions.CustomCapturePoint;
-					}
+				}
 
 					if ((uint32)(InDesiredSize.X + StartCapturePoint.X) > SourceTexture->GetSizeX() || (uint32)(InDesiredSize.Y + StartCapturePoint.Y) > SourceTexture->GetSizeY())
-					{
+				{
 						InMediaCapture->MediaState = EMediaCaptureState::Error;
-						UE_LOG(LogMediaIOCore, Error, TEXT("The capture will stop for '%s'. The Source size doesn't match with the user requested size. Requested: %d,%d  Source: %d,%d")
+					UE_LOG(LogMediaIOCore, Error, TEXT("The capture will stop for '%s'. The Source size doesn't match with the user requested size. Requested: %d,%d  Source: %d,%d")
 							, *InMediaCapture->MediaOutputName
 							, InDesiredSize.X, InDesiredSize.Y
-							, SourceTexture->GetSizeX(), SourceTexture->GetSizeY());
-					}
+						, SourceTexture->GetSizeX(), SourceTexture->GetSizeY());
 				}
 			}
+		}
 
 			if (CapturingFrame && InMediaCapture->MediaState != EMediaCaptureState::Error)
-			{
-				SCOPE_CYCLE_COUNTER(STAT_MediaCapture_RenderThread_CopyToResolve);
+		{
+			SCOPE_CYCLE_COUNTER(STAT_MediaCapture_RenderThread_CopyToResolve);
 
-				FPooledRenderTargetDesc OutputDesc = FPooledRenderTargetDesc::Create2DDesc(
+			FPooledRenderTargetDesc OutputDesc = FPooledRenderTargetDesc::Create2DDesc(
 					InMediaCapture->DesiredOutputSize,
 					InMediaCapture->DesiredOutputPixelFormat,
-					FClearValueBinding::None,
-					TexCreate_None,
-					TexCreate_RenderTargetable,
-					false);
-				TRefCountPtr<IPooledRenderTarget> ResampleTexturePooledRenderTarget;
-				GetRendererModule().RenderTargetPoolFindFreeElement(RHICmdList, OutputDesc, ResampleTexturePooledRenderTarget, TEXT("MediaCapture"));
-				const FSceneRenderTargetItem& DestRenderTarget = ResampleTexturePooledRenderTarget->GetRenderTargetItem();
-
-				// Do we need to crop
-				float ULeft = 0.0f;
-				float URight = 1.0f;
-				float VTop = 0.0f;
-				float VBottom = 1.0f;
-				FResolveParams ResolveParams;
+				FClearValueBinding::None,
+				TexCreate_None,
+				TexCreate_RenderTargetable,
+				false);
+			TRefCountPtr<IPooledRenderTarget> ResampleTexturePooledRenderTarget;
+			GetRendererModule().RenderTargetPoolFindFreeElement(RHICmdList, OutputDesc, ResampleTexturePooledRenderTarget, TEXT("MediaCapture"));
+			const FSceneRenderTargetItem& DestRenderTarget = ResampleTexturePooledRenderTarget->GetRenderTargetItem();
+
+			// Do we need to crop
+			float ULeft = 0.0f;
+			float URight = 1.0f;
+			float VTop = 0.0f;
+			float VBottom = 1.0f;
+			FResolveParams ResolveParams;
 				if (InMediaCapture->DesiredCaptureOptions.Crop != EMediaCaptureCroppingType::None)
-				{
+			{
 					switch (InMediaCapture->DesiredCaptureOptions.Crop)
-					{
-					case EMediaCaptureCroppingType::Center:
+				{
+				case EMediaCaptureCroppingType::Center:
 						ResolveParams.Rect = FResolveRect((SourceTexture->GetSizeX() - InDesiredSize.X) / 2, (SourceTexture->GetSizeY() - InDesiredSize.Y) / 2, 0, 0);
 						ResolveParams.Rect.X2 = ResolveParams.Rect.X1 + InDesiredSize.X;
 						ResolveParams.Rect.Y2 = ResolveParams.Rect.Y1 + InDesiredSize.Y;
-						break;
-					case EMediaCaptureCroppingType::TopLeft:
+					break;
+				case EMediaCaptureCroppingType::TopLeft:
 						ResolveParams.Rect = FResolveRect(0, 0, InDesiredSize.X, InDesiredSize.Y);
-						break;
-					case EMediaCaptureCroppingType::Custom:
+					break;
+				case EMediaCaptureCroppingType::Custom:
 						ResolveParams.Rect = FResolveRect(InMediaCapture->DesiredCaptureOptions.CustomCapturePoint.X, InMediaCapture->DesiredCaptureOptions.CustomCapturePoint.Y, 0, 0);
 						ResolveParams.Rect.X2 = ResolveParams.Rect.X1 + InDesiredSize.X;
 						ResolveParams.Rect.Y2 = ResolveParams.Rect.Y1 + InDesiredSize.Y;
-						break;
-					}
+					break;
+				}
 
-					ResolveParams.DestRect.X1 = 0;
+				ResolveParams.DestRect.X1 = 0;
 					ResolveParams.DestRect.X2 = InDesiredSize.X;
-					ResolveParams.DestRect.Y1 = 0;
+				ResolveParams.DestRect.Y1 = 0;
 					ResolveParams.DestRect.Y2 = InDesiredSize.Y;
 
-					ULeft = (float)ResolveParams.Rect.X1 / (float)SourceTexture->GetSizeX();
-					URight = (float)ResolveParams.Rect.X2 / (float)SourceTexture->GetSizeX();
-					VTop = (float)ResolveParams.Rect.Y1 / (float)SourceTexture->GetSizeY();
-					VBottom = (float)ResolveParams.Rect.Y2 / (float)SourceTexture->GetSizeY();
-				}
+				ULeft = (float)ResolveParams.Rect.X1 / (float)SourceTexture->GetSizeX();
+				URight = (float)ResolveParams.Rect.X2 / (float)SourceTexture->GetSizeX();
+				VTop = (float)ResolveParams.Rect.Y1 / (float)SourceTexture->GetSizeY();
+				VBottom = (float)ResolveParams.Rect.Y2 / (float)SourceTexture->GetSizeY();
+			}
 
-				{
-					SCOPED_DRAW_EVENTF(RHICmdList, MediaCapture, TEXT("MediaCapture"));
+			{
+				SCOPED_DRAW_EVENTF(RHICmdList, MediaCapture, TEXT("MediaCapture"));
 
 					if (InMediaCapture->ConversionOperation == EMediaCaptureConversionOperation::NONE)
-					{
-						// Asynchronously copy target from GPU to GPU
-						RHICmdList.CopyToResolveTarget(SourceTexture, DestRenderTarget.TargetableTexture, ResolveParams);
-					}
-					else
-					{
-						// convert the source with a draw call
-						FGraphicsPipelineStateInitializer GraphicsPSOInit;
-						FRHITexture* RenderTarget = DestRenderTarget.TargetableTexture.GetReference();
-						SetRenderTargets(RHICmdList, 1, &RenderTarget, nullptr, ESimpleRenderTargetMode::EExistingColorAndDepth, FExclusiveDepthStencil::DepthNop_StencilNop);
+				{
+					// Asynchronously copy target from GPU to GPU
+					RHICmdList.CopyToResolveTarget(SourceTexture, DestRenderTarget.TargetableTexture, ResolveParams);
+				}
+				else
+				{
+					// convert the source with a draw call
+					FGraphicsPipelineStateInitializer GraphicsPSOInit;
+					FRHITexture* RenderTarget = DestRenderTarget.TargetableTexture.GetReference();
+					SetRenderTargets(RHICmdList, 1, &RenderTarget, nullptr, ESimpleRenderTargetMode::EExistingColorAndDepth, FExclusiveDepthStencil::DepthNop_StencilNop);
 
-						RHICmdList.ApplyCachedRenderTargets(GraphicsPSOInit);
+					RHICmdList.ApplyCachedRenderTargets(GraphicsPSOInit);
 
-						GraphicsPSOInit.DepthStencilState = TStaticDepthStencilState<false, CF_Always>::GetRHI();
-						GraphicsPSOInit.RasterizerState = TStaticRasterizerState<>::GetRHI();
-						GraphicsPSOInit.BlendState = TStaticBlendStateWriteMask<CW_RGBA, CW_NONE, CW_NONE, CW_NONE, CW_NONE, CW_NONE, CW_NONE, CW_NONE>::GetRHI();
-						GraphicsPSOInit.PrimitiveType = PT_TriangleStrip;
+					GraphicsPSOInit.DepthStencilState = TStaticDepthStencilState<false, CF_Always>::GetRHI();
+					GraphicsPSOInit.RasterizerState = TStaticRasterizerState<>::GetRHI();
+					GraphicsPSOInit.BlendState = TStaticBlendStateWriteMask<CW_RGBA, CW_NONE, CW_NONE, CW_NONE, CW_NONE, CW_NONE, CW_NONE, CW_NONE>::GetRHI();
+					GraphicsPSOInit.PrimitiveType = PT_TriangleStrip;
 
-						// configure media shaders
-						auto ShaderMap = GetGlobalShaderMap(GMaxRHIFeatureLevel);
-						TShaderMapRef<FMediaShadersVS> VertexShader(ShaderMap);
+					// configure media shaders
+					auto ShaderMap = GetGlobalShaderMap(GMaxRHIFeatureLevel);
+					TShaderMapRef<FMediaShadersVS> VertexShader(ShaderMap);
 
-						GraphicsPSOInit.BoundShaderState.VertexDeclarationRHI = GMediaVertexDeclaration.VertexDeclarationRHI;
-						GraphicsPSOInit.BoundShaderState.VertexShaderRHI = GETSAFERHISHADER_VERTEX(*VertexShader);
+					GraphicsPSOInit.BoundShaderState.VertexDeclarationRHI = GMediaVertexDeclaration.VertexDeclarationRHI;
+					GraphicsPSOInit.BoundShaderState.VertexShaderRHI = GETSAFERHISHADER_VERTEX(*VertexShader);
 
-						const bool bDoLinearToSRGB = false;
+					const bool bDoLinearToSRGB = false;
 
 						switch (InMediaCapture->ConversionOperation)
+					{
+					case EMediaCaptureConversionOperation::RGBA8_TO_YUV_8BIT:
+						{
+							TShaderMapRef<FRGB8toUYVY8ConvertPS> ConvertShader(ShaderMap);
+							GraphicsPSOInit.BoundShaderState.PixelShaderRHI = GETSAFERHISHADER_PIXEL(*ConvertShader);
+							SetGraphicsPipelineState(RHICmdList, GraphicsPSOInit);
+						ConvertShader->SetParameters(RHICmdList, SourceTexture, MediaShaders::RgbToYuvRec709Full, MediaShaders::YUVOffset8bits, bDoLinearToSRGB);
+						}
+						break;
+					case EMediaCaptureConversionOperation::RGB10_TO_YUVv210_10BIT:
+						{
+							TShaderMapRef<FRGB10toYUVv210ConvertPS> ConvertShader(ShaderMap);
+							GraphicsPSOInit.BoundShaderState.PixelShaderRHI = GETSAFERHISHADER_PIXEL(*ConvertShader);
+							SetGraphicsPipelineState(RHICmdList, GraphicsPSOInit);
+						ConvertShader->SetParameters(RHICmdList, SourceTexture, MediaShaders::RgbToYuvRec709Full, MediaShaders::YUVOffset10bits, bDoLinearToSRGB);
+						}
+						break;
+					case EMediaCaptureConversionOperation::INVERT_ALPHA:
 						{
-						case EMediaCaptureConversionOperation::RGBA8_TO_YUV_8BIT:
-							{
-								TShaderMapRef<FRGB8toUYVY8ConvertPS> ConvertShader(ShaderMap);
-								GraphicsPSOInit.BoundShaderState.PixelShaderRHI = GETSAFERHISHADER_PIXEL(*ConvertShader);
-								SetGraphicsPipelineState(RHICmdList, GraphicsPSOInit);
-							ConvertShader->SetParameters(RHICmdList, SourceTexture, MediaShaders::RgbToYuvRec709Full, MediaShaders::YUVOffset8bits, bDoLinearToSRGB);
-							}
-							break;
-						case EMediaCaptureConversionOperation::RGB10_TO_YUVv210_10BIT:
-							{
-								TShaderMapRef<FRGB10toYUVv210ConvertPS> ConvertShader(ShaderMap);
-								GraphicsPSOInit.BoundShaderState.PixelShaderRHI = GETSAFERHISHADER_PIXEL(*ConvertShader);
-								SetGraphicsPipelineState(RHICmdList, GraphicsPSOInit);
-							ConvertShader->SetParameters(RHICmdList, SourceTexture, MediaShaders::RgbToYuvRec709Full, MediaShaders::YUVOffset10bits, bDoLinearToSRGB);
-							}
-							break;
-						case EMediaCaptureConversionOperation::INVERT_ALPHA:
-							{
-								TShaderMapRef<FInvertAlphaPS> ConvertShader(ShaderMap);
-								GraphicsPSOInit.BoundShaderState.PixelShaderRHI = GETSAFERHISHADER_PIXEL(*ConvertShader);
-								SetGraphicsPipelineState(RHICmdList, GraphicsPSOInit);
-								ConvertShader->SetParameters(RHICmdList, SourceTexture);
-							}
-							break;
+							TShaderMapRef<FInvertAlphaPS> ConvertShader(ShaderMap);
+							GraphicsPSOInit.BoundShaderState.PixelShaderRHI = GETSAFERHISHADER_PIXEL(*ConvertShader);
+							SetGraphicsPipelineState(RHICmdList, GraphicsPSOInit);
+							ConvertShader->SetParameters(RHICmdList, SourceTexture);
 						}
+						break;
+					}
 
-						// draw full size quad into render target
-						FVertexBufferRHIRef VertexBuffer = CreateTempMediaVertexBuffer(ULeft, URight, VTop, VBottom);
-						RHICmdList.SetStreamSource(0, VertexBuffer, 0);
+					// draw full size quad into render target
+					FVertexBufferRHIRef VertexBuffer = CreateTempMediaVertexBuffer(ULeft, URight, VTop, VBottom);
+					RHICmdList.SetStreamSource(0, VertexBuffer, 0);
 
-						// set viewport to RT size
+					// set viewport to RT size
 						RHICmdList.SetViewport(0, 0, 0.0f, InMediaCapture->DesiredOutputSize.X, InMediaCapture->DesiredOutputSize.Y, 1.0f);
-						RHICmdList.DrawPrimitive(0, 2, 1);
-						RHICmdList.TransitionResource(EResourceTransitionAccess::EReadable, DestRenderTarget.TargetableTexture);
-					}
+					RHICmdList.DrawPrimitive(0, 2, 1);
+					RHICmdList.TransitionResource(EResourceTransitionAccess::EReadable, DestRenderTarget.TargetableTexture);
 				}
+			}
 
-				// Asynchronously copy duplicate target from GPU to System Memory
+			// Asynchronously copy duplicate target from GPU to System Memory
 				RHICmdList.CopyToResolveTarget(DestRenderTarget.TargetableTexture, CapturingFrame->ReadbackTexture, FResolveParams());
 
 				CapturingFrame->bResolvedTargetRequested = true;
-			}
+		}
 
 			if (ReadyFrame && InMediaCapture->MediaState != EMediaCaptureState::Error)
-			{
+		{
 				check(ReadyFrame->ReadbackTexture.IsValid());
 
-				// Lock & read
-				void* ColorDataBuffer = nullptr;
-				int32 Width = 0, Height = 0;
-				{
-					SCOPE_CYCLE_COUNTER(STAT_MediaCapture_RenderThread_MapStaging);
+			// Lock & read
+			void* ColorDataBuffer = nullptr;
+			int32 Width = 0, Height = 0;
+			{
+				SCOPE_CYCLE_COUNTER(STAT_MediaCapture_RenderThread_MapStaging);
 					RHICmdList.MapStagingSurface(ReadyFrame->ReadbackTexture, ColorDataBuffer, Width, Height);
-				}
+			}
 
-				{
-					SCOPE_CYCLE_COUNTER(STAT_MediaCapture_RenderThread_Callback);
+			{
+				SCOPE_CYCLE_COUNTER(STAT_MediaCapture_RenderThread_Callback);
 					InMediaCapture->OnFrameCaptured_RenderingThread(ReadyFrame->CaptureBaseData, ReadyFrame->UserData, ColorDataBuffer, Width, Height);
-				}
+			}
 				ReadyFrame->bResolvedTargetRequested = false;
 
 				RHICmdList.UnmapStagingSurface(ReadyFrame->ReadbackTexture);
-			}
+		}
 
 			--InMediaCapture->WaitingForResolveCommandExecutionCounter;
 		}
