Inverse alpha before writing to file in FileMediaOutput.

diff --git a/Engine/Plugins/Media/MediaFrameworkUtilities/Source/MediaFrameworkUtilities/Private/MediaAssets/ProxyMediaOutput.cpp b/Engine/Plugins/Media/MediaFrameworkUtilities/Source/MediaFrameworkUtilities/Private/MediaAssets/ProxyMediaOutput.cpp
--- a/Engine/Plugins/Media/MediaFrameworkUtilities/Source/MediaFrameworkUtilities/Private/MediaAssets/ProxyMediaOutput.cpp
+++ b/Engine/Plugins/Media/MediaFrameworkUtilities/Source/MediaFrameworkUtilities/Private/MediaAssets/ProxyMediaOutput.cpp
@@ -73,7 +73,7 @@ EPixelFormat UProxyMediaOutput::GetRequestedPixelFormat() const
 }
 
 
-EMediaCaptureConversionOperation UProxyMediaOutput::GetConversionOperation() const
+EMediaCaptureConversionOperation UProxyMediaOutput::GetConversionOperation(EMediaCaptureSourceType InSourceType) const
 {
 	// Guard against reentrant calls.
 	if (bRequestedPixelFormatGuard)
@@ -84,7 +84,7 @@ EMediaCaptureConversionOperation UProxyMediaOutput::GetConversionOperation() con
 	TGuardValue<bool> GettingUrlGuard(bRequestedPixelFormatGuard, true);
 
 	UMediaOutput* CurrentProxy = GetMediaOutput();
-	return (CurrentProxy != nullptr) ? CurrentProxy->GetConversionOperation() : EMediaCaptureConversionOperation::NONE;
+	return (CurrentProxy != nullptr) ? CurrentProxy->GetConversionOperation(InSourceType) : EMediaCaptureConversionOperation::NONE;
 }
 
 
diff --git a/Engine/Plugins/Media/MediaFrameworkUtilities/Source/MediaFrameworkUtilities/Public/MediaAssets/ProxyMediaOutput.h b/Engine/Plugins/Media/MediaFrameworkUtilities/Source/MediaFrameworkUtilities/Public/MediaAssets/ProxyMediaOutput.h
--- a/Engine/Plugins/Media/MediaFrameworkUtilities/Source/MediaFrameworkUtilities/Public/MediaAssets/ProxyMediaOutput.h
+++ b/Engine/Plugins/Media/MediaFrameworkUtilities/Source/MediaFrameworkUtilities/Public/MediaAssets/ProxyMediaOutput.h
@@ -65,7 +65,7 @@ public:
 	virtual bool Validate(FString& OutFailureReason) const override;
 	virtual FIntPoint GetRequestedSize() const override;
 	virtual EPixelFormat GetRequestedPixelFormat() const override;
-	virtual EMediaCaptureConversionOperation GetConversionOperation() const override;
+	virtual EMediaCaptureConversionOperation GetConversionOperation(EMediaCaptureSourceType InSourceType) const override;
 
 protected:
 
diff --git a/Engine/Plugins/Media/MediaIOFramework/Source/MediaIOCore/Private/FileMediaOutput.cpp b/Engine/Plugins/Media/MediaIOFramework/Source/MediaIOCore/Private/FileMediaOutput.cpp
--- a/Engine/Plugins/Media/MediaIOFramework/Source/MediaIOCore/Private/FileMediaOutput.cpp
+++ b/Engine/Plugins/Media/MediaIOFramework/Source/MediaIOCore/Private/FileMediaOutput.cpp
@@ -76,6 +76,30 @@ EPixelFormat UFileMediaOutput::GetRequestedPixelFormat() const
 }
 
 
+EMediaCaptureConversionOperation UFileMediaOutput::GetConversionOperation(EMediaCaptureSourceType InSourceType) const
+{
+	// All formats supporting alpha
+	if (WriteOptions.Format == EDesiredImageFormat::EXR || WriteOptions.Format == EDesiredImageFormat::PNG)
+	{
+		// We invert alpha only when alpha channel as valid data when used with "passthrough tone mapper" or using a render target, otherwise we force it to 1.0f.
+		static const auto CVarPropagateAlpha = IConsoleManager::Get().FindTConsoleVariableDataInt(TEXT("r.PostProcessing.PropagateAlpha"));
+		EAlphaChannelMode::Type PropagateAlpha = EAlphaChannelMode::FromInt(CVarPropagateAlpha->GetValueOnAnyThread());
+		if ((PropagateAlpha == EAlphaChannelMode::AllowThroughTonemapper) || (InSourceType == EMediaCaptureSourceType::RENDER_TARGET))
+		{
+			return EMediaCaptureConversionOperation::INVERT_ALPHA;
+		}
+		else
+		{
+			return EMediaCaptureConversionOperation::SET_ALPHA_ONE;
+		}
+	}
+	else
+	{
+		return EMediaCaptureConversionOperation::NONE;
+	}
+}
+
+
 UMediaCapture* UFileMediaOutput::CreateMediaCaptureImpl()
 {
 	UMediaCapture* Result = NewObject<UFileMediaCapture>();
diff --git a/Engine/Plugins/Media/MediaIOFramework/Source/MediaIOCore/Private/MediaCapture.cpp b/Engine/Plugins/Media/MediaIOFramework/Source/MediaIOCore/Private/MediaCapture.cpp
--- a/Engine/Plugins/Media/MediaIOFramework/Source/MediaIOCore/Private/MediaCapture.cpp
+++ b/Engine/Plugins/Media/MediaIOFramework/Source/MediaIOCore/Private/MediaCapture.cpp
@@ -147,7 +147,7 @@ bool UMediaCapture::CaptureSceneViewport(TSharedPtr<FSceneViewport>& InSceneView
 	}
 
 	DesiredCaptureOptions = InCaptureOptions;
-	CacheMediaOutput();
+	CacheMediaOutput(EMediaCaptureSourceType::SCENE_VIEWPORT);
 
 	const bool bCurrentlyCapturing = false;
 	if (!MediaCaptureDetails::ValidateSceneViewport(InSceneViewport, DesiredCaptureOptions, DesiredSize, DesiredPixelFormat, bCurrentlyCapturing))
@@ -187,7 +187,7 @@ bool UMediaCapture::CaptureTextureRenderTarget2D(UTextureRenderTarget2D* InRende
 	}
 
 	DesiredCaptureOptions = CaptureOptions;
-	CacheMediaOutput();
+	CacheMediaOutput(EMediaCaptureSourceType::RENDER_TARGET);
 
 	const bool bCurrentlyCapturing = false;
 	if (!MediaCaptureDetails::ValidateTextureRenderTarget2D(InRenderTarget2D, DesiredCaptureOptions, DesiredSize, DesiredPixelFormat, bCurrentlyCapturing))
@@ -214,12 +214,12 @@ bool UMediaCapture::CaptureTextureRenderTarget2D(UTextureRenderTarget2D* InRende
 	return true;
 }
 
-void UMediaCapture::CacheMediaOutput()
+void UMediaCapture::CacheMediaOutput(EMediaCaptureSourceType InSourceType)
 {
 	check(MediaOutput);
 	DesiredSize = MediaOutput->GetRequestedSize();
 	DesiredPixelFormat = MediaOutput->GetRequestedPixelFormat();
-	ConversionOperation = MediaOutput->GetConversionOperation();
+	ConversionOperation = MediaOutput->GetConversionOperation(InSourceType);
 	DesiredOutputSize = GetOutputSize(DesiredSize, ConversionOperation);
 	DesiredOutputPixelFormat = GetOutputPixelFormat(DesiredPixelFormat, ConversionOperation);
 	MediaOutputName = *MediaOutput->GetName();
@@ -527,50 +527,50 @@ void UMediaCapture::OnEndFrame_GameThread()
 
 		if (!SourceTexture.IsValid())
 		{
-				InMediaCapture->MediaState = EMediaCaptureState::Error;
-				UE_LOG(LogMediaIOCore, Error, TEXT("Can't grab the Texture to capture for '%s'."), *InMediaCapture->MediaOutputName);
+			InMediaCapture->MediaState = EMediaCaptureState::Error;
+			UE_LOG(LogMediaIOCore, Error, TEXT("Can't grab the Texture to capture for '%s'."), *InMediaCapture->MediaOutputName);
 		}
-			else if (CapturingFrame)
+		else if (CapturingFrame)
 		{
-				if (InMediaCapture->DesiredPixelFormat != SourceTexture->GetFormat())
+			if (InMediaCapture->DesiredPixelFormat != SourceTexture->GetFormat())
 			{
-					InMediaCapture->MediaState = EMediaCaptureState::Error;
+				InMediaCapture->MediaState = EMediaCaptureState::Error;
 				UE_LOG(LogMediaIOCore, Error, TEXT("The capture will stop for '%s'. The Source pixel format doesn't match with the user requested pixel format. Requested: %s Source: %s")
-						, *InMediaCapture->MediaOutputName
-						, GetPixelFormatString(InMediaCapture->DesiredPixelFormat)
+					, *InMediaCapture->MediaOutputName
+					, GetPixelFormatString(InMediaCapture->DesiredPixelFormat)
 					, GetPixelFormatString(SourceTexture->GetFormat()));
 			}
-				else if (InMediaCapture->DesiredCaptureOptions.Crop == EMediaCaptureCroppingType::None)
+			else if (InMediaCapture->DesiredCaptureOptions.Crop == EMediaCaptureCroppingType::None)
 			{
-					if (InDesiredSize.X != SourceTexture->GetSizeX() || InDesiredSize.Y != SourceTexture->GetSizeY())
+				if (InDesiredSize.X != SourceTexture->GetSizeX() || InDesiredSize.Y != SourceTexture->GetSizeY())
 				{
-						InMediaCapture->MediaState = EMediaCaptureState::Error;
+					InMediaCapture->MediaState = EMediaCaptureState::Error;
 					UE_LOG(LogMediaIOCore, Error, TEXT("The capture will stop for '%s'. The Source size doesn't match with the user requested size. Requested: %d,%d  Source: %d,%d")
-							, *InMediaCapture->MediaOutputName
-							, InDesiredSize.X, InDesiredSize.Y
+						, *InMediaCapture->MediaOutputName
+						, InDesiredSize.X, InDesiredSize.Y
 						, SourceTexture->GetSizeX(), SourceTexture->GetSizeY());
 				}
 			}
 			else
 			{
 				FIntPoint StartCapturePoint = FIntPoint::ZeroValue;
-					if (InMediaCapture->DesiredCaptureOptions.Crop == EMediaCaptureCroppingType::Custom)
+				if (InMediaCapture->DesiredCaptureOptions.Crop == EMediaCaptureCroppingType::Custom)
 				{
-						StartCapturePoint = InMediaCapture->DesiredCaptureOptions.CustomCapturePoint;
+					StartCapturePoint = InMediaCapture->DesiredCaptureOptions.CustomCapturePoint;
 				}
 
-					if ((uint32)(InDesiredSize.X + StartCapturePoint.X) > SourceTexture->GetSizeX() || (uint32)(InDesiredSize.Y + StartCapturePoint.Y) > SourceTexture->GetSizeY())
+				if ((uint32)(InDesiredSize.X + StartCapturePoint.X) > SourceTexture->GetSizeX() || (uint32)(InDesiredSize.Y + StartCapturePoint.Y) > SourceTexture->GetSizeY())
 				{
-						InMediaCapture->MediaState = EMediaCaptureState::Error;
+					InMediaCapture->MediaState = EMediaCaptureState::Error;
 					UE_LOG(LogMediaIOCore, Error, TEXT("The capture will stop for '%s'. The Source size doesn't match with the user requested size. Requested: %d,%d  Source: %d,%d")
-							, *InMediaCapture->MediaOutputName
-							, InDesiredSize.X, InDesiredSize.Y
+						, *InMediaCapture->MediaOutputName
+						, InDesiredSize.X, InDesiredSize.Y
 						, SourceTexture->GetSizeX(), SourceTexture->GetSizeY());
 				}
 			}
 		}
 
-			if (CapturingFrame && InMediaCapture->MediaState != EMediaCaptureState::Error)
+		if (CapturingFrame && InMediaCapture->MediaState != EMediaCaptureState::Error)
 		{
 			SCOPE_CYCLE_COUNTER(STAT_MediaCapture_RenderThread_CopyToResolve);
 
@@ -591,29 +591,29 @@ void UMediaCapture::OnEndFrame_GameThread()
 			float VTop = 0.0f;
 			float VBottom = 1.0f;
 			FResolveParams ResolveParams;
-				if (InMediaCapture->DesiredCaptureOptions.Crop != EMediaCaptureCroppingType::None)
+			if (InMediaCapture->DesiredCaptureOptions.Crop != EMediaCaptureCroppingType::None)
 			{
-					switch (InMediaCapture->DesiredCaptureOptions.Crop)
+				switch (InMediaCapture->DesiredCaptureOptions.Crop)
 				{
 				case EMediaCaptureCroppingType::Center:
-						ResolveParams.Rect = FResolveRect((SourceTexture->GetSizeX() - InDesiredSize.X) / 2, (SourceTexture->GetSizeY() - InDesiredSize.Y) / 2, 0, 0);
-						ResolveParams.Rect.X2 = ResolveParams.Rect.X1 + InDesiredSize.X;
-						ResolveParams.Rect.Y2 = ResolveParams.Rect.Y1 + InDesiredSize.Y;
+					ResolveParams.Rect = FResolveRect((SourceTexture->GetSizeX() - InDesiredSize.X) / 2, (SourceTexture->GetSizeY() - InDesiredSize.Y) / 2, 0, 0);
+					ResolveParams.Rect.X2 = ResolveParams.Rect.X1 + InDesiredSize.X;
+					ResolveParams.Rect.Y2 = ResolveParams.Rect.Y1 + InDesiredSize.Y;
 					break;
 				case EMediaCaptureCroppingType::TopLeft:
-						ResolveParams.Rect = FResolveRect(0, 0, InDesiredSize.X, InDesiredSize.Y);
+					ResolveParams.Rect = FResolveRect(0, 0, InDesiredSize.X, InDesiredSize.Y);
 					break;
 				case EMediaCaptureCroppingType::Custom:
-						ResolveParams.Rect = FResolveRect(InMediaCapture->DesiredCaptureOptions.CustomCapturePoint.X, InMediaCapture->DesiredCaptureOptions.CustomCapturePoint.Y, 0, 0);
-						ResolveParams.Rect.X2 = ResolveParams.Rect.X1 + InDesiredSize.X;
-						ResolveParams.Rect.Y2 = ResolveParams.Rect.Y1 + InDesiredSize.Y;
+					ResolveParams.Rect = FResolveRect(InMediaCapture->DesiredCaptureOptions.CustomCapturePoint.X, InMediaCapture->DesiredCaptureOptions.CustomCapturePoint.Y, 0, 0);
+					ResolveParams.Rect.X2 = ResolveParams.Rect.X1 + InDesiredSize.X;
+					ResolveParams.Rect.Y2 = ResolveParams.Rect.Y1 + InDesiredSize.Y;
 					break;
 				}
 
 				ResolveParams.DestRect.X1 = 0;
-					ResolveParams.DestRect.X2 = InDesiredSize.X;
+				ResolveParams.DestRect.X2 = InDesiredSize.X;
 				ResolveParams.DestRect.Y1 = 0;
-					ResolveParams.DestRect.Y2 = InDesiredSize.Y;
+				ResolveParams.DestRect.Y2 = InDesiredSize.Y;
 
 				ULeft = (float)ResolveParams.Rect.X1 / (float)SourceTexture->GetSizeX();
 				URight = (float)ResolveParams.Rect.X2 / (float)SourceTexture->GetSizeX();
@@ -624,7 +624,7 @@ void UMediaCapture::OnEndFrame_GameThread()
 			{
 				SCOPED_DRAW_EVENTF(RHICmdList, MediaCapture, TEXT("MediaCapture"));
 
-					if (InMediaCapture->ConversionOperation == EMediaCaptureConversionOperation::NONE)
+				if (InMediaCapture->ConversionOperation == EMediaCaptureConversionOperation::NONE)
 				{
 					// Asynchronously copy target from GPU to GPU
 					RHICmdList.CopyToResolveTarget(SourceTexture, DestRenderTarget.TargetableTexture, ResolveParams);
@@ -652,14 +652,14 @@ void UMediaCapture::OnEndFrame_GameThread()
 
 					const bool bDoLinearToSRGB = false;
 
-						switch (InMediaCapture->ConversionOperation)
+					switch (InMediaCapture->ConversionOperation)
 					{
 					case EMediaCaptureConversionOperation::RGBA8_TO_YUV_8BIT:
 						{
 							TShaderMapRef<FRGB8toUYVY8ConvertPS> ConvertShader(ShaderMap);
 							GraphicsPSOInit.BoundShaderState.PixelShaderRHI = GETSAFERHISHADER_PIXEL(*ConvertShader);
 							SetGraphicsPipelineState(RHICmdList, GraphicsPSOInit);
-						ConvertShader->SetParameters(RHICmdList, SourceTexture, MediaShaders::RgbToYuvRec709Full, MediaShaders::YUVOffset8bits, bDoLinearToSRGB);
+							ConvertShader->SetParameters(RHICmdList, SourceTexture, MediaShaders::RgbToYuvRec709Full, MediaShaders::YUVOffset8bits, bDoLinearToSRGB);
 						}
 						break;
 					case EMediaCaptureConversionOperation::RGB10_TO_YUVv210_10BIT:
@@ -667,7 +667,7 @@ void UMediaCapture::OnEndFrame_GameThread()
 							TShaderMapRef<FRGB10toYUVv210ConvertPS> ConvertShader(ShaderMap);
 							GraphicsPSOInit.BoundShaderState.PixelShaderRHI = GETSAFERHISHADER_PIXEL(*ConvertShader);
 							SetGraphicsPipelineState(RHICmdList, GraphicsPSOInit);
-						ConvertShader->SetParameters(RHICmdList, SourceTexture, MediaShaders::RgbToYuvRec709Full, MediaShaders::YUVOffset10bits, bDoLinearToSRGB);
+							ConvertShader->SetParameters(RHICmdList, SourceTexture, MediaShaders::RgbToYuvRec709Full, MediaShaders::YUVOffset10bits, bDoLinearToSRGB);
 						}
 						break;
 					case EMediaCaptureConversionOperation::INVERT_ALPHA:
@@ -678,6 +678,14 @@ void UMediaCapture::OnEndFrame_GameThread()
 							ConvertShader->SetParameters(RHICmdList, SourceTexture);
 						}
 						break;
+					case EMediaCaptureConversionOperation::SET_ALPHA_ONE:
+						{
+							TShaderMapRef<FSetAlphaOnePS> ConvertShader(ShaderMap);
+							GraphicsPSOInit.BoundShaderState.PixelShaderRHI = GETSAFERHISHADER_PIXEL(*ConvertShader);
+							SetGraphicsPipelineState(RHICmdList, GraphicsPSOInit);
+							ConvertShader->SetParameters(RHICmdList, SourceTexture);
+						}
+						break;
 					}
 
 					// draw full size quad into render target
@@ -685,42 +693,41 @@ void UMediaCapture::OnEndFrame_GameThread()
 					RHICmdList.SetStreamSource(0, VertexBuffer, 0);
 
 					// set viewport to RT size
-						RHICmdList.SetViewport(0, 0, 0.0f, InMediaCapture->DesiredOutputSize.X, InMediaCapture->DesiredOutputSize.Y, 1.0f);
+					RHICmdList.SetViewport(0, 0, 0.0f, InMediaCapture->DesiredOutputSize.X, InMediaCapture->DesiredOutputSize.Y, 1.0f);
 					RHICmdList.DrawPrimitive(0, 2, 1);
 					RHICmdList.TransitionResource(EResourceTransitionAccess::EReadable, DestRenderTarget.TargetableTexture);
 				}
 			}
 
 			// Asynchronously copy duplicate target from GPU to System Memory
-				RHICmdList.CopyToResolveTarget(DestRenderTarget.TargetableTexture, CapturingFrame->ReadbackTexture, FResolveParams());
+			RHICmdList.CopyToResolveTarget(DestRenderTarget.TargetableTexture, CapturingFrame->ReadbackTexture, FResolveParams());
 
-				CapturingFrame->bResolvedTargetRequested = true;
+			CapturingFrame->bResolvedTargetRequested = true;
 		}
 
-			if (ReadyFrame && InMediaCapture->MediaState != EMediaCaptureState::Error)
+		if (ReadyFrame && InMediaCapture->MediaState != EMediaCaptureState::Error)
 		{
-				check(ReadyFrame->ReadbackTexture.IsValid());
+			check(ReadyFrame->ReadbackTexture.IsValid());
 
 			// Lock & read
 			void* ColorDataBuffer = nullptr;
 			int32 Width = 0, Height = 0;
 			{
 				SCOPE_CYCLE_COUNTER(STAT_MediaCapture_RenderThread_MapStaging);
-					RHICmdList.MapStagingSurface(ReadyFrame->ReadbackTexture, ColorDataBuffer, Width, Height);
+				RHICmdList.MapStagingSurface(ReadyFrame->ReadbackTexture, ColorDataBuffer, Width, Height);
 			}
 
 			{
 				SCOPE_CYCLE_COUNTER(STAT_MediaCapture_RenderThread_Callback);
-					InMediaCapture->OnFrameCaptured_RenderingThread(ReadyFrame->CaptureBaseData, ReadyFrame->UserData, ColorDataBuffer, Width, Height);
+				InMediaCapture->OnFrameCaptured_RenderingThread(ReadyFrame->CaptureBaseData, ReadyFrame->UserData, ColorDataBuffer, Width, Height);
 			}
-				ReadyFrame->bResolvedTargetRequested = false;
+			ReadyFrame->bResolvedTargetRequested = false;
 
-				RHICmdList.UnmapStagingSurface(ReadyFrame->ReadbackTexture);
+			RHICmdList.UnmapStagingSurface(ReadyFrame->ReadbackTexture);
 		}
 
-			--InMediaCapture->WaitingForResolveCommandExecutionCounter;
-		}
-	);
+		--InMediaCapture->WaitingForResolveCommandExecutionCounter;
+	});
 }
 
 /* namespace MediaCaptureDetails implementation
diff --git a/Engine/Plugins/Media/MediaIOFramework/Source/MediaIOCore/Public/FileMediaOutput.h b/Engine/Plugins/Media/MediaIOFramework/Source/MediaIOCore/Public/FileMediaOutput.h
--- a/Engine/Plugins/Media/MediaIOFramework/Source/MediaIOCore/Public/FileMediaOutput.h
+++ b/Engine/Plugins/Media/MediaIOFramework/Source/MediaIOCore/Public/FileMediaOutput.h
@@ -66,6 +66,7 @@ public:
 	virtual bool Validate(FString& FailureReason) const override;
 	virtual FIntPoint GetRequestedSize() const override;
 	virtual EPixelFormat GetRequestedPixelFormat() const override;
+	virtual EMediaCaptureConversionOperation GetConversionOperation(EMediaCaptureSourceType InSourceType) const override;
 
 protected:
 	virtual UMediaCapture* CreateMediaCaptureImpl() override;
diff --git a/Engine/Plugins/Media/MediaIOFramework/Source/MediaIOCore/Public/MediaCapture.h b/Engine/Plugins/Media/MediaIOFramework/Source/MediaIOCore/Public/MediaCapture.h
--- a/Engine/Plugins/Media/MediaIOFramework/Source/MediaIOCore/Public/MediaCapture.h
+++ b/Engine/Plugins/Media/MediaIOFramework/Source/MediaIOCore/Public/MediaCapture.h
@@ -212,7 +212,7 @@ protected:
 private:
 	void InitializeResolveTarget(int32 InNumberOfBuffers);
 	void OnEndFrame_GameThread();
-	void CacheMediaOutput();
+	void CacheMediaOutput(EMediaCaptureSourceType InSourceType);
 	FIntPoint GetOutputSize(const FIntPoint & InSize, const EMediaCaptureConversionOperation & InConversionOperation) const;
 	EPixelFormat GetOutputPixelFormat(const EPixelFormat & InPixelFormat, const EMediaCaptureConversionOperation & InConversionOperation) const;
 
diff --git a/Engine/Plugins/Media/MediaIOFramework/Source/MediaIOCore/Public/MediaOutput.h b/Engine/Plugins/Media/MediaIOFramework/Source/MediaIOCore/Public/MediaOutput.h
--- a/Engine/Plugins/Media/MediaIOFramework/Source/MediaIOCore/Public/MediaOutput.h
+++ b/Engine/Plugins/Media/MediaIOFramework/Source/MediaIOCore/Public/MediaOutput.h
@@ -19,6 +19,13 @@ enum class EMediaCaptureConversionOperation : int8
 	RGBA8_TO_YUV_8BIT,
 	RGB10_TO_YUVv210_10BIT,
 	INVERT_ALPHA,
+	SET_ALPHA_ONE,
+};
+
+enum class EMediaCaptureSourceType : int8
+{
+	RENDER_TARGET,
+	SCENE_VIEWPORT,
 };
 
 /**
@@ -56,7 +63,7 @@ public:
 public:
 	virtual FIntPoint GetRequestedSize() const PURE_VIRTUAL(UMediaOutput::GetRequestedSize, return FIntPoint::ZeroValue; );
 	virtual EPixelFormat GetRequestedPixelFormat() const PURE_VIRTUAL(UMediaOutput::GetRequestedPixelFormat, return EPixelFormat::PF_Unknown; );
-	virtual EMediaCaptureConversionOperation GetConversionOperation() const { return EMediaCaptureConversionOperation::NONE; }
+	virtual EMediaCaptureConversionOperation GetConversionOperation(EMediaCaptureSourceType InSourceType) const { return EMediaCaptureConversionOperation::NONE; }
 
 protected:
 	virtual UMediaCapture* CreateMediaCaptureImpl() { return nullptr; }
diff --git a/Engine/Shaders/Private/MediaShaders.usf b/Engine/Shaders/Private/MediaShaders.usf
--- a/Engine/Shaders/Private/MediaShaders.usf
+++ b/Engine/Shaders/Private/MediaShaders.usf
@@ -358,3 +358,14 @@ void InvertAlphaPS(
 	OutColor.xyzw = float4(Color.rgb, 1 - Color.a);
 }
 
+
+// shader to set alpha to 1.0
+void SetAlphaOnePS(
+	float4 InPosition : SV_POSITION,
+	float2 InUV : TEXCOORD0,
+	out float4 OutColor : SV_Target0)
+{
+	float4 Color = SetAlphaOneUB.Texture.Sample(SetAlphaOneUB.SamplerP, InUV);
+	OutColor.xyzw = float4(Color.rgb, 1.0f);
+}
+
diff --git a/Engine/Source/Runtime/UtilityShaders/Private/MediaShaders.cpp b/Engine/Source/Runtime/UtilityShaders/Private/MediaShaders.cpp
--- a/Engine/Source/Runtime/UtilityShaders/Private/MediaShaders.cpp
+++ b/Engine/Source/Runtime/UtilityShaders/Private/MediaShaders.cpp
@@ -539,4 +539,26 @@ void FInvertAlphaPS::SetParameters(FRHICommandList& CommandList, TRefCountPtr<FR
 }
 
 
+/* FSetAlphaOnePS shader
+ *****************************************************************************/
+
+BEGIN_GLOBAL_SHADER_PARAMETER_STRUCT(FSetAlphaOneUB, )
+SHADER_PARAMETER_TEXTURE(Texture2D, Texture)
+SHADER_PARAMETER_SAMPLER(SamplerState, SamplerP)
+END_GLOBAL_SHADER_PARAMETER_STRUCT()
+
+IMPLEMENT_GLOBAL_SHADER_PARAMETER_STRUCT(FSetAlphaOneUB, "SetAlphaOneUB");
+IMPLEMENT_SHADER_TYPE(, FSetAlphaOnePS, TEXT("/Engine/Private/MediaShaders.usf"), TEXT("SetAlphaOnePS"), SF_Pixel);
+
 
+void FSetAlphaOnePS::SetParameters(FRHICommandList& CommandList, TRefCountPtr<FRHITexture2D> RGBATexture)
+{
+	FSetAlphaOneUB UB;
+	{
+		UB.SamplerP = TStaticSamplerState<SF_Point>::GetRHI();
+		UB.Texture = RGBATexture;
+	}
+
+	TUniformBufferRef<FSetAlphaOneUB> Data = TUniformBufferRef<FSetAlphaOneUB>::CreateUniformBufferImmediate(UB, UniformBuffer_SingleFrame);
+	SetUniformBufferParameter(CommandList, GetPixelShader(), GetUniformBufferParameter<FSetAlphaOneUB>(), Data);
+}
diff --git a/Engine/Source/Runtime/UtilityShaders/Public/MediaShaders.h b/Engine/Source/Runtime/UtilityShaders/Public/MediaShaders.h
--- a/Engine/Source/Runtime/UtilityShaders/Public/MediaShaders.h
+++ b/Engine/Source/Runtime/UtilityShaders/Public/MediaShaders.h
@@ -637,3 +637,30 @@ public:
 
 	UTILITYSHADERS_API void SetParameters(FRHICommandList& RHICmdList, TRefCountPtr<FRHITexture2D> RGBATexture);
 };
+
+
+/**
+ * Pixel shader to set alpha component to 1.0f
+ *
+ * This shader expects a single texture in RGBA 8 or 10 bits format.
+ */
+class FSetAlphaOnePS
+	: public FGlobalShader
+{
+	DECLARE_EXPORTED_SHADER_TYPE(FSetAlphaOnePS, Global, UTILITYSHADERS_API);
+
+public:
+
+	static bool ShouldCompilePermutation(const FGlobalShaderPermutationParameters& Parameters)
+	{
+		return IsFeatureLevelSupported(Parameters.Platform, ERHIFeatureLevel::ES2);
+	}
+
+	FSetAlphaOnePS() { }
+
+	FSetAlphaOnePS(const ShaderMetaType::CompiledShaderInitializerType& Initializer)
+		: FGlobalShader(Initializer)
+	{ }
+
+	UTILITYSHADERS_API void SetParameters(FRHICommandList& RHICmdList, TRefCountPtr<FRHITexture2D> RGBATexture);
+};
