Add delegate to notify when the media capture state changed.
In Media Capture panel, reset the viewport when the capture stop.

diff --git a/Engine/Plugins/Media/MediaFrameworkUtilities/Source/MediaFrameworkUtilitiesEditor/Private/CaptureTab/SMediaFrameworkCaptureOutputWidget.cpp b/Engine/Plugins/Media/MediaFrameworkUtilities/Source/MediaFrameworkUtilitiesEditor/Private/CaptureTab/SMediaFrameworkCaptureOutputWidget.cpp
--- a/Engine/Plugins/Media/MediaFrameworkUtilities/Source/MediaFrameworkUtilitiesEditor/Private/CaptureTab/SMediaFrameworkCaptureOutputWidget.cpp
+++ b/Engine/Plugins/Media/MediaFrameworkUtilities/Source/MediaFrameworkUtilitiesEditor/Private/CaptureTab/SMediaFrameworkCaptureOutputWidget.cpp
@@ -625,11 +625,16 @@ void SMediaFrameworkCaptureCurrentViewportWidget::StartOutput()
 
 			if (SceneViewport.IsValid())
 			{
-				GEditor->OnLevelViewportClientListChanged().AddRaw(this, &SMediaFrameworkCaptureCurrentViewportWidget::OnLevelViewportClientListChanged);
+				GEditor->OnLevelViewportClientListChanged().AddSP(this, &SMediaFrameworkCaptureCurrentViewportWidget::OnLevelViewportClientListChanged);
 				EditorSceneViewport = SceneViewport;
 				FIntPoint TargetSize = MediaOutputPtr->GetRequestedSize();
 				SceneViewport->SetFixedViewportSize(TargetSize.X, TargetSize.Y);
-				MediaCapture->CaptureSceneViewport(SceneViewport, CaptureOptions);
+				MediaCapture->OnStateChangedNative.AddSP(this, &SMediaFrameworkCaptureCurrentViewportWidget::OnMediaCaptureStateChanged);
+
+				if (!MediaCapture->CaptureSceneViewport(SceneViewport, CaptureOptions))
+				{
+					ShutdownViewport();
+				}
 			}
 		}
 	}
@@ -657,6 +662,18 @@ void SMediaFrameworkCaptureCurrentViewportWidget::OnLevelViewportClientListChang
 	}
 }
 
+void SMediaFrameworkCaptureCurrentViewportWidget::OnMediaCaptureStateChanged()
+{
+	if (MediaCapture.IsValid())
+	{
+		EMediaCaptureState State = MediaCapture->GetState();
+		if (State != EMediaCaptureState::Capturing && State != EMediaCaptureState::Preparing)
+		{
+			ShutdownViewport();
+		}
+	}
+}
+
 void SMediaFrameworkCaptureCurrentViewportWidget::OnPrePIE()
 {
 	if (MediaCapture.IsValid())
@@ -701,6 +718,7 @@ void SMediaFrameworkCaptureCurrentViewportWidget::ShutdownViewport()
 
 	LevelViewport.Reset();
 	EditorSceneViewport.Reset();
+	MediaCapture.Reset();
 }
 
 #undef LOCTEXT_NAMESPACE
diff --git a/Engine/Plugins/Media/MediaFrameworkUtilities/Source/MediaFrameworkUtilitiesEditor/Private/CaptureTab/SMediaFrameworkCaptureOutputWidget.h b/Engine/Plugins/Media/MediaFrameworkUtilities/Source/MediaFrameworkUtilitiesEditor/Private/CaptureTab/SMediaFrameworkCaptureOutputWidget.h
--- a/Engine/Plugins/Media/MediaFrameworkUtilities/Source/MediaFrameworkUtilitiesEditor/Private/CaptureTab/SMediaFrameworkCaptureOutputWidget.h
+++ b/Engine/Plugins/Media/MediaFrameworkUtilities/Source/MediaFrameworkUtilitiesEditor/Private/CaptureTab/SMediaFrameworkCaptureOutputWidget.h
@@ -176,6 +176,7 @@ public:
 private:
 	void ShutdownViewport();
 	void OnLevelViewportClientListChanged();
+	void OnMediaCaptureStateChanged();
 
 private:
 	TWeakPtr<FSceneViewport> EditorSceneViewport;
diff --git a/Engine/Plugins/Media/MediaIOFramework/Source/MediaIOCore/Private/FileMediaCapture.cpp b/Engine/Plugins/Media/MediaIOFramework/Source/MediaIOCore/Private/FileMediaCapture.cpp
--- a/Engine/Plugins/Media/MediaIOFramework/Source/MediaIOCore/Private/FileMediaCapture.cpp
+++ b/Engine/Plugins/Media/MediaIOFramework/Source/MediaIOCore/Private/FileMediaCapture.cpp
@@ -14,7 +14,7 @@ void UFileMediaCapture::OnFrameCaptured_RenderingThread(const FCaptureBaseData&
 	IImageWriteQueueModule* ImageWriteQueueModule = FModuleManager::Get().GetModulePtr<IImageWriteQueueModule>("ImageWriteQueue");
 	if (ImageWriteQueueModule == nullptr)
 	{
-		MediaState = EMediaCaptureState::Error;
+		SetState(EMediaCaptureState::Error);
 		return;
 	}
 
@@ -61,7 +61,7 @@ bool UFileMediaCapture::CaptureSceneViewportImpl(TSharedPtr<FSceneViewport>& InS
 	FModuleManager::Get().LoadModuleChecked<IImageWriteQueueModule>("ImageWriteQueue");
 	CacheMediaOutputValues();
 
-	MediaState = EMediaCaptureState::Capturing;
+	SetState(EMediaCaptureState::Capturing);
 	return true;
 }
 
@@ -71,7 +71,7 @@ bool UFileMediaCapture::CaptureRenderTargetImpl(UTextureRenderTarget2D* InRender
 	FModuleManager::Get().LoadModuleChecked<IImageWriteQueueModule>("ImageWriteQueue");
 	CacheMediaOutputValues();
 
-	MediaState = EMediaCaptureState::Capturing;
+	SetState(EMediaCaptureState::Capturing);
 	return true;
 }
 
diff --git a/Engine/Plugins/Media/MediaIOFramework/Source/MediaIOCore/Private/MediaCapture.cpp b/Engine/Plugins/Media/MediaIOFramework/Source/MediaIOCore/Private/MediaCapture.cpp
--- a/Engine/Plugins/Media/MediaIOFramework/Source/MediaIOCore/Private/MediaCapture.cpp
+++ b/Engine/Plugins/Media/MediaIOFramework/Source/MediaIOCore/Private/MediaCapture.cpp
@@ -2,6 +2,7 @@
 
 #include "MediaCapture.h"
 
+#include "Async/Async.h"
 #include "Engine/GameEngine.h"
 #include "Engine/RendererSettings.h"
 #include "Engine/TextureRenderTarget2D.h"
@@ -11,13 +12,14 @@
 #include "MediaOutput.h"
 #include "MediaShaders.h"
 #include "Misc/App.h"
+#include "Misc/ScopeLock.h"
 #include "PipelineStateCache.h"
 #include "RendererInterface.h"
 #include "RenderUtils.h"
 #include "RHIStaticStates.h"
 #include "SceneUtils.h"
 #include "Slate/SceneViewport.h"
-#include "Misc/ScopeLock.h"
+#include "UObject/WeakObjectPtrTemplates.h"
 
 #if WITH_EDITOR
 #include "Editor.h"
@@ -84,9 +86,9 @@ FMediaCaptureOptions::FMediaCaptureOptions()
 
 UMediaCapture::UMediaCapture(const FObjectInitializer& ObjectInitializer)
 	: Super(ObjectInitializer)
-	, MediaState(EMediaCaptureState::Stopped)
 	, CurrentResolvedTargetIndex(0)
 	, NumberOfCaptureFrame(2)
+	, MediaState(EMediaCaptureState::Stopped)
 	, DesiredSize(1280, 720)
 	, DesiredPixelFormat(EPixelFormat::PF_A2B10G10R10)
 	, DesiredOutputSize(1280, 720)
@@ -100,7 +102,7 @@ UMediaCapture::UMediaCapture(const FObjectInitializer& ObjectInitializer)
 
 void UMediaCapture::BeginDestroy()
 {
-	if (MediaState == EMediaCaptureState::Capturing || MediaState == EMediaCaptureState::Preparing)
+	if (GetState() == EMediaCaptureState::Capturing || GetState() == EMediaCaptureState::Preparing)
 	{
 		UE_LOG(LogMediaIOCore, Warning, TEXT("%s will be destroyed and the capture was not stopped."), *GetName());
 	}
@@ -156,10 +158,10 @@ bool UMediaCapture::CaptureSceneViewport(TSharedPtr<FSceneViewport>& InSceneView
 		return false;
 	}
 
-	MediaState = EMediaCaptureState::Preparing;
+	SetState(EMediaCaptureState::Preparing);
 	if (!CaptureSceneViewportImpl(InSceneViewport))
 	{
-		MediaState = EMediaCaptureState::Stopped;
+		SetState(EMediaCaptureState::Stopped);
 		MediaCaptureDetails::ShowSlateNotification();
 		return false;
 	}
@@ -196,10 +198,10 @@ bool UMediaCapture::CaptureTextureRenderTarget2D(UTextureRenderTarget2D* InRende
 		return false;
 	}
 
-	MediaState = EMediaCaptureState::Preparing;
+	SetState(EMediaCaptureState::Preparing);
 	if (!CaptureRenderTargetImpl(InRenderTarget2D))
 	{
-		MediaState = EMediaCaptureState::Stopped;
+		SetState(EMediaCaptureState::Stopped);
 		MediaCaptureDetails::ShowSlateNotification();
 		return false;
 	}
@@ -328,23 +330,23 @@ void UMediaCapture::StopCapture(bool bAllowPendingFrameToBeProcess)
 {
 	check(IsInGameThread());
 
-	if (MediaState != EMediaCaptureState::StopRequested && MediaState != EMediaCaptureState::Capturing)
+	if (GetState() != EMediaCaptureState::StopRequested && GetState() != EMediaCaptureState::Capturing)
 	{
 		bAllowPendingFrameToBeProcess = false;
 	}
 
 	if (bAllowPendingFrameToBeProcess)
 	{
-		if (MediaState != EMediaCaptureState::Stopped && MediaState != EMediaCaptureState::StopRequested)
+		if (GetState() != EMediaCaptureState::Stopped && GetState() != EMediaCaptureState::StopRequested)
 		{
-			MediaState = EMediaCaptureState::StopRequested;
+			SetState(EMediaCaptureState::StopRequested);
 		}
 	}
 	else
 	{
-		if (MediaState != EMediaCaptureState::Stopped)
+		if (GetState() != EMediaCaptureState::Stopped)
 		{
-			MediaState = EMediaCaptureState::Stopped;
+			SetState(EMediaCaptureState::Stopped);
 
 			FCoreDelegates::OnEndFrame.RemoveAll(this);
 
@@ -376,11 +378,41 @@ void UMediaCapture::SetMediaOutput(UMediaOutput* InMediaOutput)
 	}
 }
 
+void UMediaCapture::SetState(EMediaCaptureState InNewState)
+{
+	if (MediaState != InNewState)
+	{
+		MediaState = InNewState;
+		if (IsInGameThread())
+		{
+			BroadcastStateChanged();
+		}
+		else
+		{
+			TWeakObjectPtr<UMediaCapture> Self = this;
+			AsyncTask(ENamedThreads::GameThread, [Self]
+			{
+				UMediaCapture* MediaCapture = Self.Get();
+				if (MediaCapture)
+				{
+					MediaCapture->BroadcastStateChanged();
+				}
+			});
+		}
+	}
+}
+
+void UMediaCapture::BroadcastStateChanged()
+{
+	OnStateChanged.Broadcast();
+	OnStateChangedNative.Broadcast();
+}
+
 bool UMediaCapture::HasFinishedProcessing() const
 {
 	return WaitingForResolveCommandExecutionCounter == 0
-		|| MediaState == EMediaCaptureState::Error
-		|| MediaState == EMediaCaptureState::Stopped;
+		|| GetState() == EMediaCaptureState::Error
+		|| GetState() == EMediaCaptureState::Stopped;
 }
 
 void UMediaCapture::InitializeResolveTarget(int32 InNumberOfBuffers)
@@ -482,6 +514,7 @@ void UMediaCapture::OnEndFrame_GameThread()
 	TWeakPtr<FSceneViewport> InCapturingSceneViewport = CapturingSceneViewport;
 	FTextureRenderTargetResource* InTextureRenderTargetResource = nullptr;
 	FIntPoint InDesiredSize = DesiredSize;
+	FMediaCaptureStateChangedSignature InOnStateChanged = OnStateChanged;
 	UMediaCapture* InMediaCapture = this;
 
 	{
@@ -494,7 +527,7 @@ void UMediaCapture::OnEndFrame_GameThread()
 
 	// RenderCommand to be executed on the RenderThread
 	ENQUEUE_RENDER_COMMAND(FMediaOutputCaptureFrameCreateTexture)(
-		[CapturingFrame, ReadyFrame, InCapturingSceneViewport, InTextureRenderTargetResource, InDesiredSize, InMediaCapture](FRHICommandListImmediate& RHICmdList)
+		[CapturingFrame, ReadyFrame, InCapturingSceneViewport, InTextureRenderTargetResource, InDesiredSize, InOnStateChanged, InMediaCapture](FRHICommandListImmediate& RHICmdList)
 	{
 		FTexture2DRHIRef SourceTexture;
 		{
@@ -527,14 +560,14 @@ void UMediaCapture::OnEndFrame_GameThread()
 
 		if (!SourceTexture.IsValid())
 		{
-			InMediaCapture->MediaState = EMediaCaptureState::Error;
+			InMediaCapture->SetState(EMediaCaptureState::Error);
 			UE_LOG(LogMediaIOCore, Error, TEXT("Can't grab the Texture to capture for '%s'."), *InMediaCapture->MediaOutputName);
 		}
 		else if (CapturingFrame)
 		{
 			if (InMediaCapture->DesiredPixelFormat != SourceTexture->GetFormat())
 			{
-				InMediaCapture->MediaState = EMediaCaptureState::Error;
+				InMediaCapture->SetState(EMediaCaptureState::Error);
 				UE_LOG(LogMediaIOCore, Error, TEXT("The capture will stop for '%s'. The Source pixel format doesn't match with the user requested pixel format. Requested: %s Source: %s")
 					, *InMediaCapture->MediaOutputName
 					, GetPixelFormatString(InMediaCapture->DesiredPixelFormat)
@@ -544,7 +577,7 @@ void UMediaCapture::OnEndFrame_GameThread()
 			{
 				if (InDesiredSize.X != SourceTexture->GetSizeX() || InDesiredSize.Y != SourceTexture->GetSizeY())
 				{
-					InMediaCapture->MediaState = EMediaCaptureState::Error;
+					InMediaCapture->SetState(EMediaCaptureState::Error);
 					UE_LOG(LogMediaIOCore, Error, TEXT("The capture will stop for '%s'. The Source size doesn't match with the user requested size. Requested: %d,%d  Source: %d,%d")
 						, *InMediaCapture->MediaOutputName
 						, InDesiredSize.X, InDesiredSize.Y
@@ -561,7 +594,7 @@ void UMediaCapture::OnEndFrame_GameThread()
 
 				if ((uint32)(InDesiredSize.X + StartCapturePoint.X) > SourceTexture->GetSizeX() || (uint32)(InDesiredSize.Y + StartCapturePoint.Y) > SourceTexture->GetSizeY())
 				{
-					InMediaCapture->MediaState = EMediaCaptureState::Error;
+					InMediaCapture->SetState(EMediaCaptureState::Error);
 					UE_LOG(LogMediaIOCore, Error, TEXT("The capture will stop for '%s'. The Source size doesn't match with the user requested size. Requested: %d,%d  Source: %d,%d")
 						, *InMediaCapture->MediaOutputName
 						, InDesiredSize.X, InDesiredSize.Y
@@ -570,13 +603,13 @@ void UMediaCapture::OnEndFrame_GameThread()
 			}
 		}
 
-		if (CapturingFrame && InMediaCapture->MediaState != EMediaCaptureState::Error)
+		if (CapturingFrame && InMediaCapture->GetState() != EMediaCaptureState::Error)
 		{
 			SCOPE_CYCLE_COUNTER(STAT_MediaCapture_RenderThread_CopyToResolve);
 
 			FPooledRenderTargetDesc OutputDesc = FPooledRenderTargetDesc::Create2DDesc(
-					InMediaCapture->DesiredOutputSize,
-					InMediaCapture->DesiredOutputPixelFormat,
+				InMediaCapture->DesiredOutputSize,
+				InMediaCapture->DesiredOutputPixelFormat,
 				FClearValueBinding::None,
 				TexCreate_None,
 				TexCreate_RenderTargetable,
@@ -705,7 +738,7 @@ void UMediaCapture::OnEndFrame_GameThread()
 			CapturingFrame->bResolvedTargetRequested = true;
 		}
 
-		if (ReadyFrame && InMediaCapture->MediaState != EMediaCaptureState::Error)
+		if (ReadyFrame && InMediaCapture->GetState() != EMediaCaptureState::Error)
 		{
 			check(ReadyFrame->ReadbackTexture.IsValid());
 
diff --git a/Engine/Plugins/Media/MediaIOFramework/Source/MediaIOCore/Public/MediaCapture.h b/Engine/Plugins/Media/MediaIOFramework/Source/MediaIOCore/Public/MediaCapture.h
--- a/Engine/Plugins/Media/MediaIOFramework/Source/MediaIOCore/Public/MediaCapture.h
+++ b/Engine/Plugins/Media/MediaIOFramework/Source/MediaIOCore/Public/MediaCapture.h
@@ -22,8 +22,8 @@ class UTextureRenderTarget2D;
 /**
  * Possible states of media capture.
  */
-UENUM()
-enum class EMediaCaptureState
+UENUM(BlueprintType)
+enum class EMediaCaptureState : uint8
 {
 	/** Unrecoverable error occurred during capture. */
 	Error,
@@ -50,7 +50,7 @@ class FMediaCaptureUserData
 /**
  * Type of cropping 
  */
-UENUM()
+UENUM(BlueprintType)
 enum class EMediaCaptureCroppingType : uint8
 {
 	/** Do not crop the captured image. */
@@ -86,6 +86,12 @@ public:
 	FIntPoint CustomCapturePoint;
 };
 
+
+/** Delegate signatures */
+DECLARE_DYNAMIC_MULTICAST_DELEGATE(FMediaCaptureStateChangedSignature);
+DECLARE_MULTICAST_DELEGATE(FMediaCaptureStateChangedSignatureNative);
+
+
 /**
  * Abstract base class for media capture.
  *
@@ -173,6 +179,16 @@ public:
 	/** Check whether this capture has any processing left to do. */
 	virtual bool HasFinishedProcessing() const;
 
+	/** Called when the state of the capture changed. */
+	UPROPERTY(BlueprintAssignable, Category = "Media|Output")
+	FMediaCaptureStateChangedSignature OnStateChanged;
+
+	/**
+	 * Called when the state of the capture changed.
+	 * The callback is called on the game thread. Note that the change may occur on the rendering thread.
+	 */
+	FMediaCaptureStateChangedSignatureNative OnStateChangedNative;
+
 public:
 	//~ UObject interface
 	virtual void BeginDestroy() override;
@@ -202,19 +218,19 @@ protected:
 	UTextureRenderTarget2D* GetTextureRenderTarget() { return CapturingRenderTarget; }
 	TSharedPtr<FSceneViewport> GetCapturingSceneViewport() { return CapturingSceneViewport.Pin(); }
 	EMediaCaptureConversionOperation GetConversionOperation() const { return ConversionOperation; }
+	void SetState(EMediaCaptureState InNewState);
 
 protected:
 	UPROPERTY(Transient)
 	UMediaOutput* MediaOutput;
 
-	EMediaCaptureState MediaState;
-
 private:
 	void InitializeResolveTarget(int32 InNumberOfBuffers);
 	void OnEndFrame_GameThread();
 	void CacheMediaOutput(EMediaCaptureSourceType InSourceType);
 	FIntPoint GetOutputSize(const FIntPoint & InSize, const EMediaCaptureConversionOperation & InConversionOperation) const;
 	EPixelFormat GetOutputPixelFormat(const EPixelFormat & InPixelFormat, const EMediaCaptureConversionOperation & InConversionOperation) const;
+	void BroadcastStateChanged();
 
 private:
 	struct FCaptureFrame
@@ -230,6 +246,7 @@ private:
 	TArray<FCaptureFrame> CaptureFrames;
 	int32 CurrentResolvedTargetIndex;
 	int32 NumberOfCaptureFrame;
+	EMediaCaptureState MediaState;
 
 	UPROPERTY(Transient)
 	UTextureRenderTarget2D* CapturingRenderTarget;
