Timecode

- Moving Timecode adjustments logic in AJA library (duplicate frames, odd frames, etc.)
- Removing BurnTimecode option

diff --git a/Engine/Plugins/Media/MediaIOFramework/Source/MediaIOCore/Private/MediaCapture.cpp b/Engine/Plugins/Media/MediaIOFramework/Source/MediaIOCore/Private/MediaCapture.cpp
--- a/Engine/Plugins/Media/MediaIOFramework/Source/MediaIOCore/Private/MediaCapture.cpp
+++ b/Engine/Plugins/Media/MediaIOFramework/Source/MediaIOCore/Private/MediaCapture.cpp
@@ -504,6 +504,7 @@ void UMediaCapture::OnEndFrame_GameThread()
 		}
 
 		CapturingFrame->CaptureBaseData.SourceFrameTimecode = FApp::GetTimecode();
+		CapturingFrame->CaptureBaseData.SourceFrameTimecodeFramerate = FApp::GetTimecodeFrameRate();
 		CapturingFrame->CaptureBaseData.SourceFrameNumberRenderThread = GFrameNumber;
 		CapturingFrame->UserData = GetCaptureFrameUserData_GameThread();
 	}
diff --git a/Engine/Plugins/Media/MediaIOFramework/Source/MediaIOCore/Private/Player/MediaIOCorePlayerBase.cpp b/Engine/Plugins/Media/MediaIOFramework/Source/MediaIOCore/Private/Player/MediaIOCorePlayerBase.cpp
--- a/Engine/Plugins/Media/MediaIOFramework/Source/MediaIOCore/Private/Player/MediaIOCorePlayerBase.cpp
+++ b/Engine/Plugins/Media/MediaIOFramework/Source/MediaIOCore/Private/Player/MediaIOCorePlayerBase.cpp
@@ -142,9 +142,8 @@ void FMediaIOCorePlayerBase::TickTimeManagement()
 {
 	if (bUseTimeSynchronization)
 	{
-		FTimecode Timecode = FApp::GetTimecode();
-		FFrameRate FrameRate = FApp::GetTimecodeFrameRate();
-		CurrentTime = FTimespan(0, Timecode.Hours, Timecode.Minutes, Timecode.Seconds, static_cast<int32>((ETimespan::TicksPerSecond * Timecode.Frames) / FrameRate.AsDecimal()) * ETimespan::NanosecondsPerTick);
+		const FTimecode Timecode = FApp::GetTimecode();
+		CurrentTime = Timecode.ToTimespan(FApp::GetTimecodeFrameRate());
 	}
 	else
 	{
diff --git a/Engine/Plugins/Media/MediaIOFramework/Source/MediaIOCore/Public/MediaCapture.h b/Engine/Plugins/Media/MediaIOFramework/Source/MediaIOCore/Public/MediaCapture.h
--- a/Engine/Plugins/Media/MediaIOFramework/Source/MediaIOCore/Public/MediaCapture.h
+++ b/Engine/Plugins/Media/MediaIOFramework/Source/MediaIOCore/Public/MediaCapture.h
@@ -209,6 +209,7 @@ protected:
 		FCaptureBaseData();
 
 		FTimecode SourceFrameTimecode;
+		FFrameRate SourceFrameTimecodeFramerate;
 		uint32 SourceFrameNumberRenderThread;
 	};
 	virtual TSharedPtr<FMediaCaptureUserData> GetCaptureFrameUserData_GameThread() { return TSharedPtr<FMediaCaptureUserData>(); }
