Add an option to resize the source media when capturing.
Add an option have capture the RHI texture for media capture.

diff --git a/Engine/Plugins/Media/MediaFrameworkUtilities/Source/MediaFrameworkUtilitiesEditor/Private/CaptureTab/SMediaFrameworkCaptureOutputWidget.cpp b/Engine/Plugins/Media/MediaFrameworkUtilities/Source/MediaFrameworkUtilitiesEditor/Private/CaptureTab/SMediaFrameworkCaptureOutputWidget.cpp
--- a/Engine/Plugins/Media/MediaFrameworkUtilities/Source/MediaFrameworkUtilitiesEditor/Private/CaptureTab/SMediaFrameworkCaptureOutputWidget.cpp
+++ b/Engine/Plugins/Media/MediaFrameworkUtilities/Source/MediaFrameworkUtilitiesEditor/Private/CaptureTab/SMediaFrameworkCaptureOutputWidget.cpp
@@ -562,11 +562,6 @@ void SMediaFrameworkCaptureCurrentViewportWidget::Construct(const FArguments& In
 }
 END_SLATE_FUNCTION_BUILD_OPTIMIZATION
 
-SMediaFrameworkCaptureCurrentViewportWidget::~SMediaFrameworkCaptureCurrentViewportWidget()
-{
-	ShutdownViewport();
-}
-
 void SMediaFrameworkCaptureCurrentViewportWidget::StartOutput()
 {
 	ShutdownViewport();
@@ -627,8 +622,6 @@ void SMediaFrameworkCaptureCurrentViewportWidget::StartOutput()
 			{
 				GEditor->OnLevelViewportClientListChanged().AddSP(this, &SMediaFrameworkCaptureCurrentViewportWidget::OnLevelViewportClientListChanged);
 				EditorSceneViewport = SceneViewport;
-				FIntPoint TargetSize = MediaOutputPtr->GetRequestedSize();
-				SceneViewport->SetFixedViewportSize(TargetSize.X, TargetSize.Y);
 				MediaCapture->OnStateChangedNative.AddSP(this, &SMediaFrameworkCaptureCurrentViewportWidget::OnMediaCaptureStateChanged);
 
 				if (!MediaCapture->CaptureSceneViewport(SceneViewport, CaptureOptions))
@@ -640,6 +633,15 @@ void SMediaFrameworkCaptureCurrentViewportWidget::StartOutput()
 	}
 }
 
+void SMediaFrameworkCaptureCurrentViewportWidget::StopOutput()
+{
+	if (MediaCapture.IsValid())
+	{
+		MediaCapture->StopCapture(false);
+	}
+	ShutdownViewport();
+}
+
 void SMediaFrameworkCaptureCurrentViewportWidget::OnLevelViewportClientListChanged()
 {
 	bool bFound = false;
@@ -676,32 +678,19 @@ void SMediaFrameworkCaptureCurrentViewportWidget::OnMediaCaptureStateChanged()
 
 void SMediaFrameworkCaptureCurrentViewportWidget::OnPrePIE()
 {
-	if (MediaCapture.IsValid())
-	{
-		MediaCapture->StopCapture(false);
-	}
-	ShutdownViewport();
+	StopOutput();
 }
 
 void SMediaFrameworkCaptureCurrentViewportWidget::OnPrePIEEnded()
 {
-	if (MediaCapture.IsValid())
-	{
-		MediaCapture->StopCapture(false);
-	}
-	ShutdownViewport();
+	StopOutput();
 }
 
 void SMediaFrameworkCaptureCurrentViewportWidget::ShutdownViewport()
 {
-	TSharedPtr<FSceneViewport> EditorSceneViewportPin = EditorSceneViewport.Pin();
-	if (EditorSceneViewportPin.IsValid() && EditorSceneViewportPin->GetViewportWidget().IsValid())
-	{
-		EditorSceneViewportPin->SetFixedViewportSize(0, 0);
-	}
-
 	GEditor->OnLevelViewportClientListChanged().RemoveAll(this);
 
+	TSharedPtr<FSceneViewport> EditorSceneViewportPin = EditorSceneViewport.Pin();
 	TSharedPtr<ILevelViewport> LevelViewportPin = LevelViewport.Pin();
 	if (LevelViewportPin.IsValid() && LevelViewportPin->GetSharedActiveViewport() == EditorSceneViewportPin)
 	{
diff --git a/Engine/Plugins/Media/MediaFrameworkUtilities/Source/MediaFrameworkUtilitiesEditor/Private/CaptureTab/SMediaFrameworkCaptureOutputWidget.h b/Engine/Plugins/Media/MediaFrameworkUtilities/Source/MediaFrameworkUtilitiesEditor/Private/CaptureTab/SMediaFrameworkCaptureOutputWidget.h
--- a/Engine/Plugins/Media/MediaFrameworkUtilities/Source/MediaFrameworkUtilitiesEditor/Private/CaptureTab/SMediaFrameworkCaptureOutputWidget.h
+++ b/Engine/Plugins/Media/MediaFrameworkUtilities/Source/MediaFrameworkUtilitiesEditor/Private/CaptureTab/SMediaFrameworkCaptureOutputWidget.h
@@ -153,7 +153,7 @@ public:
 	SLATE_END_ARGS()
 
 	void Construct(const FArguments& InArgs);
-	virtual ~SMediaFrameworkCaptureCurrentViewportWidget();
+	virtual void StopOutput() override;
 
 	void StartOutput();
 
diff --git a/Engine/Plugins/Media/MediaFrameworkUtilities/Source/MediaFrameworkUtilitiesEditor/Private/MediaFrameworkWorldSettingsAssetUserData.cpp b/Engine/Plugins/Media/MediaFrameworkUtilities/Source/MediaFrameworkUtilitiesEditor/Private/MediaFrameworkWorldSettingsAssetUserData.cpp
--- a/Engine/Plugins/Media/MediaFrameworkUtilities/Source/MediaFrameworkUtilitiesEditor/Private/MediaFrameworkWorldSettingsAssetUserData.cpp
+++ b/Engine/Plugins/Media/MediaFrameworkUtilities/Source/MediaFrameworkUtilitiesEditor/Private/MediaFrameworkWorldSettingsAssetUserData.cpp
@@ -25,6 +25,12 @@ FMediaFrameworkCaptureRenderTargetCameraOutputInfo::FMediaFrameworkCaptureRender
 }
 
 
+UMediaFrameworkWorldSettingsAssetUserData::UMediaFrameworkWorldSettingsAssetUserData()
+{
+	CurrentViewportMediaOutput.CaptureOptions.bResizeSourceBuffer = true;
+}
+
+
 PRAGMA_DISABLE_DEPRECATION_WARNINGS
 void UMediaFrameworkWorldSettingsAssetUserData::Serialize(FArchive& Ar)
 {
diff --git a/Engine/Plugins/Media/MediaFrameworkUtilities/Source/MediaFrameworkUtilitiesEditor/Private/MediaFrameworkWorldSettingsAssetUserData.h b/Engine/Plugins/Media/MediaFrameworkUtilities/Source/MediaFrameworkUtilitiesEditor/Private/MediaFrameworkWorldSettingsAssetUserData.h
--- a/Engine/Plugins/Media/MediaFrameworkUtilities/Source/MediaFrameworkUtilitiesEditor/Private/MediaFrameworkWorldSettingsAssetUserData.h
+++ b/Engine/Plugins/Media/MediaFrameworkUtilities/Source/MediaFrameworkUtilitiesEditor/Private/MediaFrameworkWorldSettingsAssetUserData.h
@@ -95,6 +95,8 @@ class UMediaFrameworkWorldSettingsAssetUserData : public UAssetUserData
 	GENERATED_BODY()
 
 public:
+	UMediaFrameworkWorldSettingsAssetUserData();
+
 	UPROPERTY(EditAnywhere, Category="Media Render Target Capture", meta=(ShowOnlyInnerProperties))
 	TArray<FMediaFrameworkCaptureRenderTargetCameraOutputInfo> RenderTargetCaptures;
 
@@ -104,7 +106,7 @@ public:
 	/**
 	 * Capture the current viewport. It may be the level editor active viewport or a PIE instance launch with "New Editor Window PIE".
 	 * @note The behavior is different from MediaCapture.CaptureActiveSceneViewport. Here we can capture the editor viewport (since we are in the editor).
-	 * @note If the viewport is the level editor active viewport, then 1-all inputs will be disabled 2-the viewport size will be fixed 3-the viewport will always rendered.
+	 * @note If the viewport is the level editor active viewport, then all inputs will be disabled and the viewport will always rendered.
 	 */
 	UPROPERTY(EditAnywhere, Category="Media Current Viewport Capture", meta=(DisplayName="Current Viewport"))
 	FMediaFrameworkCaptureCurrentViewportOutputInfo CurrentViewportMediaOutput;
diff --git a/Engine/Plugins/Media/MediaIOFramework/Source/MediaIOCore/Private/FileMediaOutput.cpp b/Engine/Plugins/Media/MediaIOFramework/Source/MediaIOCore/Private/FileMediaOutput.cpp
--- a/Engine/Plugins/Media/MediaIOFramework/Source/MediaIOCore/Private/FileMediaOutput.cpp
+++ b/Engine/Plugins/Media/MediaIOFramework/Source/MediaIOCore/Private/FileMediaOutput.cpp
@@ -51,7 +51,7 @@ FIntPoint UFileMediaOutput::GetRequestedSize() const
 		return DesiredSize;
 	}
 
-	return FIntPoint(GSystemResolution.ResX, GSystemResolution.ResY);
+	return UMediaOutput::RequestCaptureSourceSize;
 }
 
 
diff --git a/Engine/Plugins/Media/MediaIOFramework/Source/MediaIOCore/Private/MediaCapture.cpp b/Engine/Plugins/Media/MediaIOFramework/Source/MediaIOCore/Private/MediaCapture.cpp
--- a/Engine/Plugins/Media/MediaIOFramework/Source/MediaIOCore/Private/MediaCapture.cpp
+++ b/Engine/Plugins/Media/MediaIOFramework/Source/MediaIOCore/Private/MediaCapture.cpp
@@ -95,7 +95,10 @@ UMediaCapture::UMediaCapture(const FObjectInitializer& ObjectInitializer)
 	, DesiredOutputPixelFormat(EPixelFormat::PF_A2B10G10R10)
 	, ConversionOperation(EMediaCaptureConversionOperation::NONE)
 	, MediaOutputName(TEXT("[undefined]"))
+	, bUseRequestedTargetSize(false)
 	, bResolvedTargetInitialized(false)
+	, bShouldCaptureRHITexture(false)
+	, bViewportHasFixedViewportSize(false)
 	, WaitingForResolveCommandExecutionCounter(0)
 {
 }
@@ -151,9 +154,21 @@ bool UMediaCapture::CaptureSceneViewport(TSharedPtr<FSceneViewport>& InSceneView
 	DesiredCaptureOptions = InCaptureOptions;
 	CacheMediaOutput(EMediaCaptureSourceType::SCENE_VIEWPORT);
 
+	if (bUseRequestedTargetSize)
+	{
+		DesiredSize = InSceneViewport->GetSize();
+	}
+	else if (DesiredCaptureOptions.bResizeSourceBuffer)
+	{
+		SetFixedViewportSize(InSceneViewport);
+	}
+
+	CacheOutputOptions();
+
 	const bool bCurrentlyCapturing = false;
 	if (!MediaCaptureDetails::ValidateSceneViewport(InSceneViewport, DesiredCaptureOptions, DesiredSize, DesiredPixelFormat, bCurrentlyCapturing))
 	{
+		ResetFixedViewportSize(InSceneViewport, false);
 		MediaCaptureDetails::ShowSlateNotification();
 		return false;
 	}
@@ -161,12 +176,13 @@ bool UMediaCapture::CaptureSceneViewport(TSharedPtr<FSceneViewport>& InSceneView
 	SetState(EMediaCaptureState::Preparing);
 	if (!CaptureSceneViewportImpl(InSceneViewport))
 	{
+		ResetFixedViewportSize(InSceneViewport, false);
 		SetState(EMediaCaptureState::Stopped);
 		MediaCaptureDetails::ShowSlateNotification();
 		return false;
 	}
 
-	//no lock required the command on the render thread is not active
+	//no lock required, the command on the render thread is not active
 	CapturingSceneViewport = InSceneViewport;
 
 	InitializeResolveTarget(MediaOutput->NumberOfTextureBuffers);
@@ -191,6 +207,17 @@ bool UMediaCapture::CaptureTextureRenderTarget2D(UTextureRenderTarget2D* InRende
 	DesiredCaptureOptions = CaptureOptions;
 	CacheMediaOutput(EMediaCaptureSourceType::RENDER_TARGET);
 
+	if (bUseRequestedTargetSize)
+	{
+		DesiredSize = FIntPoint(InRenderTarget2D->SizeX, InRenderTarget2D->SizeY);
+	}
+	else if (DesiredCaptureOptions.bResizeSourceBuffer)
+	{
+		InRenderTarget2D->ResizeTarget(DesiredSize.X, DesiredSize.Y);
+	}
+
+	CacheOutputOptions();
+
 	const bool bCurrentlyCapturing = false;
 	if (!MediaCaptureDetails::ValidateTextureRenderTarget2D(InRenderTarget2D, DesiredCaptureOptions, DesiredSize, DesiredPixelFormat, bCurrentlyCapturing))
 	{
@@ -222,9 +249,14 @@ void UMediaCapture::CacheMediaOutput(EMediaCaptureSourceType InSourceType)
 	DesiredSize = MediaOutput->GetRequestedSize();
 	DesiredPixelFormat = MediaOutput->GetRequestedPixelFormat();
 	ConversionOperation = MediaOutput->GetConversionOperation(InSourceType);
+}
+
+void UMediaCapture::CacheOutputOptions()
+{
 	DesiredOutputSize = GetOutputSize(DesiredSize, ConversionOperation);
 	DesiredOutputPixelFormat = GetOutputPixelFormat(DesiredPixelFormat, ConversionOperation);
 	MediaOutputName = *MediaOutput->GetName();
+	bShouldCaptureRHITexture = ShouldCaptureRHITexture();
 }
 
 FIntPoint UMediaCapture::GetOutputSize(const FIntPoint & InSize, const EMediaCaptureConversionOperation & InConversionOperation) const
@@ -266,10 +298,15 @@ bool UMediaCapture::UpdateSceneViewport(TSharedPtr<FSceneViewport>& InSceneViewp
 
 	check(IsInGameThread());
 
-	const bool bCurrentlyCapturing = true;
+	if (!bUseRequestedTargetSize && DesiredCaptureOptions.bResizeSourceBuffer)
+	{
+		SetFixedViewportSize(InSceneViewport);
+	}
 
+	const bool bCurrentlyCapturing = true;
 	if (!MediaCaptureDetails::ValidateSceneViewport(InSceneViewport, DesiredCaptureOptions, DesiredSize, DesiredPixelFormat, bCurrentlyCapturing))
 	{
+		ResetFixedViewportSize(InSceneViewport, false);
 		StopCapture(false);
 		MediaCaptureDetails::ShowSlateNotification();
 		return false;
@@ -277,6 +314,7 @@ bool UMediaCapture::UpdateSceneViewport(TSharedPtr<FSceneViewport>& InSceneViewp
 
 	if (!UpdateSceneViewportImpl(InSceneViewport))
 	{
+		ResetFixedViewportSize(InSceneViewport, false);
 		StopCapture(false);
 		MediaCaptureDetails::ShowSlateNotification();
 		return false;
@@ -284,6 +322,7 @@ bool UMediaCapture::UpdateSceneViewport(TSharedPtr<FSceneViewport>& InSceneViewp
 
 	{
 		FScopeLock Lock(&AccessingCapturingSource);
+		ResetFixedViewportSize(CapturingSceneViewport.Pin(), true);
 		CapturingSceneViewport = InSceneViewport;
 		CapturingRenderTarget = nullptr;
 	}
@@ -302,6 +341,11 @@ bool UMediaCapture::UpdateTextureRenderTarget2D(UTextureRenderTarget2D * InRende
 
 	check(IsInGameThread());
 
+	if (!bUseRequestedTargetSize && DesiredCaptureOptions.bResizeSourceBuffer)
+	{
+		InRenderTarget2D->ResizeTarget(DesiredSize.X, DesiredSize.Y);
+	}
+
 	const bool bCurrentlyCapturing = true;
 	if (!MediaCaptureDetails::ValidateTextureRenderTarget2D(InRenderTarget2D, DesiredCaptureOptions, DesiredSize, DesiredPixelFormat, bCurrentlyCapturing))
 	{
@@ -319,6 +363,7 @@ bool UMediaCapture::UpdateTextureRenderTarget2D(UTextureRenderTarget2D * InRende
 
 	{
 		FScopeLock Lock(&AccessingCapturingSource);
+		ResetFixedViewportSize(CapturingSceneViewport.Pin(), true);
 		CapturingRenderTarget = InRenderTarget2D;
 		CapturingSceneViewport.Reset();
 	}
@@ -355,6 +400,7 @@ void UMediaCapture::StopCapture(bool bAllowPendingFrameToBeProcess)
 				FlushRenderingCommands();
 			}
 			StopCaptureImpl(bAllowPendingFrameToBeProcess);
+			ResetFixedViewportSize(CapturingSceneViewport.Pin(), false);
 
 			CapturingRenderTarget = nullptr;
 			CapturingSceneViewport.Reset();
@@ -408,6 +454,25 @@ void UMediaCapture::BroadcastStateChanged()
 	OnStateChangedNative.Broadcast();
 }
 
+void UMediaCapture::SetFixedViewportSize(TSharedPtr<FSceneViewport> InSceneViewport)
+{
+	InSceneViewport->SetFixedViewportSize(DesiredSize.X, DesiredSize.Y);
+	bViewportHasFixedViewportSize = true;
+}
+
+void UMediaCapture::ResetFixedViewportSize(TSharedPtr<FSceneViewport> InViewport, bool bInFlushRenderingCommands)
+{
+	if (bViewportHasFixedViewportSize && InViewport.IsValid())
+	{
+		if (bInFlushRenderingCommands && WaitingForResolveCommandExecutionCounter > 0)
+		{
+			FlushRenderingCommands();
+		}
+		InViewport->SetFixedViewportSize(0, 0);
+		bViewportHasFixedViewportSize = false;
+	}
+}
+
 bool UMediaCapture::HasFinishedProcessing() const
 {
 	return WaitingForResolveCommandExecutionCounter == 0
@@ -417,29 +482,39 @@ bool UMediaCapture::HasFinishedProcessing() const
 
 void UMediaCapture::InitializeResolveTarget(int32 InNumberOfBuffers)
 {
+	if (bShouldCaptureRHITexture)
+	{
+		// No buffer is needed if the callback is with the RHI Texture
+		InNumberOfBuffers = 1;
+	}
+
 	NumberOfCaptureFrame = InNumberOfBuffers;
 	check(CaptureFrames.Num() == 0);
 	CaptureFrames.AddDefaulted(InNumberOfBuffers);
 
-	UMediaCapture* This = this;
-	ENQUEUE_RENDER_COMMAND(MediaOutputCaptureFrameCreateTexture)(
-		[This](FRHICommandListImmediate& RHICmdList)
-		{
-			FRHIResourceCreateInfo CreateInfo;
-			for (int32 Index = 0; Index < This->NumberOfCaptureFrame; ++Index)
+	// Only create CPU readback texture when we are using the CPU callback
+	if (!bShouldCaptureRHITexture)
+	{
+		UMediaCapture* This = this;
+		ENQUEUE_RENDER_COMMAND(MediaOutputCaptureFrameCreateTexture)(
+			[This](FRHICommandListImmediate& RHICmdList)
 			{
-				This->CaptureFrames[Index].ReadbackTexture = RHICreateTexture2D(
-					This->DesiredOutputSize.X,
-					This->DesiredOutputSize.Y,
-					This->DesiredOutputPixelFormat,
-					1,
-					1,
-					TexCreate_CPUReadback,
-					CreateInfo
-				);
-			}
-			This->bResolvedTargetInitialized = true;
-		});
+				FRHIResourceCreateInfo CreateInfo;
+				for (int32 Index = 0; Index < This->NumberOfCaptureFrame; ++Index)
+				{
+					This->CaptureFrames[Index].ReadbackTexture = RHICreateTexture2D(
+						This->DesiredOutputSize.X,
+						This->DesiredOutputSize.Y,
+						This->DesiredOutputPixelFormat,
+						1,
+						1,
+						TexCreate_CPUReadback,
+						CreateInfo
+					);
+				}
+				This->bResolvedTargetInitialized = true;
+			});
+	}
 }
 
 bool UMediaCapture::ValidateMediaOutput() const
@@ -735,13 +810,21 @@ void UMediaCapture::OnEndFrame_GameThread()
 				}
 			}
 
-			// Asynchronously copy duplicate target from GPU to System Memory
-			RHICmdList.CopyToResolveTarget(DestRenderTarget.TargetableTexture, CapturingFrame->ReadbackTexture, FResolveParams());
-
-			CapturingFrame->bResolvedTargetRequested = true;
+			if (InMediaCapture->bShouldCaptureRHITexture)
+			{
+				SCOPE_CYCLE_COUNTER(STAT_MediaCapture_RenderThread_Callback);
+				InMediaCapture->OnRHITextureCaptured_RenderingThread(CapturingFrame->CaptureBaseData, CapturingFrame->UserData, DestRenderTarget.TargetableTexture);
+				CapturingFrame->bResolvedTargetRequested = false;
+			}
+			else
+			{
+				// Asynchronously copy duplicate target from GPU to System Memory
+				RHICmdList.CopyToResolveTarget(DestRenderTarget.TargetableTexture, CapturingFrame->ReadbackTexture, FResolveParams());
+				CapturingFrame->bResolvedTargetRequested = true;
+			}
 		}
 
-		if (ReadyFrame && InMediaCapture->GetState() != EMediaCaptureState::Error)
+		if (!InMediaCapture->bShouldCaptureRHITexture && ReadyFrame && InMediaCapture->GetState() != EMediaCaptureState::Error)
 		{
 			check(ReadyFrame->ReadbackTexture.IsValid());
 
diff --git a/Engine/Plugins/Media/MediaIOFramework/Source/MediaIOCore/Private/MediaOutput.cpp b/Engine/Plugins/Media/MediaIOFramework/Source/MediaIOCore/Private/MediaOutput.cpp
--- a/Engine/Plugins/Media/MediaIOFramework/Source/MediaIOCore/Private/MediaOutput.cpp
+++ b/Engine/Plugins/Media/MediaIOFramework/Source/MediaIOCore/Private/MediaOutput.cpp
@@ -5,7 +5,9 @@
 #include "MediaCapture.h"
 #include "MediaIOCoreModule.h"
 
-/* IMediaOptions interface
+const FIntPoint UMediaOutput::RequestCaptureSourceSize = FIntPoint::ZeroValue;
+
+/* UMediaOutput
  *****************************************************************************/
 
 UMediaOutput::UMediaOutput(const FObjectInitializer& ObjectInitializer)
diff --git a/Engine/Plugins/Media/MediaIOFramework/Source/MediaIOCore/Public/MediaCapture.h b/Engine/Plugins/Media/MediaIOFramework/Source/MediaIOCore/Public/MediaCapture.h
--- a/Engine/Plugins/Media/MediaIOFramework/Source/MediaIOCore/Public/MediaCapture.h
+++ b/Engine/Plugins/Media/MediaIOFramework/Source/MediaIOCore/Public/MediaCapture.h
@@ -82,8 +82,17 @@ public:
 	 * Crop the captured SceneViewport or TextureRenderTarget2D to the desired size.
 	 * @note Only valid when Crop is set to Custom.
 	 */
-	UPROPERTY(BlueprintReadWrite, EditAnywhere, Category = "MediaCapture")
+	UPROPERTY(BlueprintReadWrite, EditAnywhere, Category="MediaCapture")
 	FIntPoint CustomCapturePoint;
+
+	/**
+	 * When the capture start, resize the source buffer to the desired size.
+	 * @note Only valid when a size is specified by the MediaOutput.
+	 * @note For viewport, the window size will not change. Only the viewport will be resized.
+	 * @note For RenderTarget, the asset will be modified and resized to the desired size.
+	 */
+	UPROPERTY(BlueprintReadWrite, EditAnywhere, Category="MediaCapture")
+	bool bResizeSourceBuffer;
 };
 
 
@@ -212,9 +221,30 @@ protected:
 		FFrameRate SourceFrameTimecodeFramerate;
 		uint32 SourceFrameNumberRenderThread;
 	};
+
+	/**
+	 * Capture the data that will pass along to the callback.
+	 * @note The capture is done on the Render Thread but triggered from the Game Thread.
+	 */
 	virtual TSharedPtr<FMediaCaptureUserData> GetCaptureFrameUserData_GameThread() { return TSharedPtr<FMediaCaptureUserData>(); }
+
+	/** Should we call OnFrameCaptured_RenderingThread() with a RHI Texture -or- copy the memory to CPU ram and call OnFrameCaptured_RenderingThread(). */
+	virtual bool ShouldCaptureRHITexture() const { return false; }
+
+	/**
+	 * Callback when the buffer was successfully copied to CPU ram.
+	 * The callback in called from a critical point. If you intend to process the buffer, do so in another thread.
+	 * The buffer is only valid for the duration of the callback.
+	 */
 	virtual void OnFrameCaptured_RenderingThread(const FCaptureBaseData& InBaseData, TSharedPtr<FMediaCaptureUserData> InUserData, void* InBuffer, int32 Width, int32 Height) { }
 
+	/**
+	 * Callback when the buffer was successfully copied on the GPU ram.
+	 * The callback in called from a critical point. If you intend to process the texture, do so in another thread.
+	 * The texture is valid for the duration of the callback.
+	 */
+	virtual void OnRHITextureCaptured_RenderingThread(const FCaptureBaseData& InBaseData, TSharedPtr<FMediaCaptureUserData> InUserData, FTextureRHIRef InTexture) { }
+
 protected:
 	UTextureRenderTarget2D* GetTextureRenderTarget() { return CapturingRenderTarget; }
 	TSharedPtr<FSceneViewport> GetCapturingSceneViewport() { return CapturingSceneViewport.Pin(); }
@@ -229,9 +259,12 @@ private:
 	void InitializeResolveTarget(int32 InNumberOfBuffers);
 	void OnEndFrame_GameThread();
 	void CacheMediaOutput(EMediaCaptureSourceType InSourceType);
+	void CacheOutputOptions();
 	FIntPoint GetOutputSize(const FIntPoint & InSize, const EMediaCaptureConversionOperation & InConversionOperation) const;
 	EPixelFormat GetOutputPixelFormat(const EPixelFormat & InPixelFormat, const EMediaCaptureConversionOperation & InConversionOperation) const;
 	void BroadcastStateChanged();
+	void SetFixedViewportSize(TSharedPtr<FSceneViewport> InSceneViewport);
+	void ResetFixedViewportSize(TSharedPtr<FSceneViewport> InViewport, bool bInFlushRenderingCommands);
 
 private:
 	struct FCaptureFrame
@@ -261,7 +294,10 @@ private:
 	FMediaCaptureOptions DesiredCaptureOptions;
 	EMediaCaptureConversionOperation ConversionOperation;
 	FString MediaOutputName;
+	bool bUseRequestedTargetSize;
 
 	bool bResolvedTargetInitialized;
+	bool bShouldCaptureRHITexture;
+	bool bViewportHasFixedViewportSize;
 	TAtomic<int32> WaitingForResolveCommandExecutionCounter;
 };
diff --git a/Engine/Plugins/Media/MediaIOFramework/Source/MediaIOCore/Public/MediaOutput.h b/Engine/Plugins/Media/MediaIOFramework/Source/MediaIOCore/Public/MediaOutput.h
--- a/Engine/Plugins/Media/MediaIOFramework/Source/MediaIOCore/Public/MediaOutput.h
+++ b/Engine/Plugins/Media/MediaIOFramework/Source/MediaIOCore/Public/MediaOutput.h
@@ -44,6 +44,7 @@ public:
 	 * Number of texture used to transfer the texture from the GPU to the system memory.
 	 * A smaller number is most likely to block the GPU (wait for the transfer to complete).
 	 * A bigger number is most likely to increase latency.
+	 * @note Some Capture are not are executed on the GPU. If it's the case then no buffer will be needed and no buffer will be created.
 	 */
 	UPROPERTY(EditAnywhere, AdvancedDisplay, Category="Output", meta=(ClampMin=1, ClampMax=4))
 	int32 NumberOfTextureBuffers;
@@ -61,8 +62,25 @@ public:
 	virtual bool Validate(FString& OutFailureReason) const;
 
 public:
-	virtual FIntPoint GetRequestedSize() const PURE_VIRTUAL(UMediaOutput::GetRequestedSize, return FIntPoint::ZeroValue; );
+
+	static const FIntPoint RequestCaptureSourceSize;
+
+	/**
+	 * The size of the buffer we wish to capture.
+	 * The size of the buffer can not change during the capture.
+	 * Return UMediaOutput::RequestCaptureSourceSize if you wish to take the buffer size as the requested size.
+	 */
+	virtual FIntPoint GetRequestedSize() const PURE_VIRTUAL(UMediaOutput::GetRequestedSize, return RequestCaptureSourceSize; );
+
+	/**
+	 * The pixel format of the buffer we wish to capture.
+	 * Some conversion are available. See EMediaCaptureConversionOperation
+	 */
 	virtual EPixelFormat GetRequestedPixelFormat() const PURE_VIRTUAL(UMediaOutput::GetRequestedPixelFormat, return EPixelFormat::PF_Unknown; );
+
+	/**
+	 * The conversion we wish to accomplish on the GPU before the DMA transfer occurs.
+	 */
 	virtual EMediaCaptureConversionOperation GetConversionOperation(EMediaCaptureSourceType InSourceType) const { return EMediaCaptureConversionOperation::NONE; }
 
 protected:
