Fix non-initialized warning in MediaCapture.

diff --git a/Engine/Plugins/Media/MediaIOFramework/Source/MediaIOCore/Private/MediaCapture.cpp b/Engine/Plugins/Media/MediaIOFramework/Source/MediaIOCore/Private/MediaCapture.cpp
--- a/Engine/Plugins/Media/MediaIOFramework/Source/MediaIOCore/Private/MediaCapture.cpp
+++ b/Engine/Plugins/Media/MediaIOFramework/Source/MediaIOCore/Private/MediaCapture.cpp
@@ -76,6 +76,7 @@ UMediaCapture::FCaptureFrame::FCaptureFrame()
 FMediaCaptureOptions::FMediaCaptureOptions()
 	: Crop(EMediaCaptureCroppingType::None)
 	, CustomCapturePoint(FIntPoint::ZeroValue)
+	, bResizeSourceBuffer(false)
 {
 
 }
