Pad filename string with '0' instead of spaces.

diff --git a/Engine/Plugins/Media/MediaIOFramework/Source/MediaIOCore/Private/FileMediaCapture.cpp b/Engine/Plugins/Media/MediaIOFramework/Source/MediaIOCore/Private/FileMediaCapture.cpp
--- a/Engine/Plugins/Media/MediaIOFramework/Source/MediaIOCore/Private/FileMediaCapture.cpp
+++ b/Engine/Plugins/Media/MediaIOFramework/Source/MediaIOCore/Private/FileMediaCapture.cpp
@@ -20,7 +20,7 @@ void UFileMediaCapture::OnFrameCaptured_RenderingThread(const FCaptureBaseData&
 
 	TUniquePtr<FImageWriteTask> ImageTask = MakeUnique<FImageWriteTask>();
 	ImageTask->Format = ImageFormat;
-	ImageTask->Filename = FString::Printf(TEXT("%s%5d"), *BaseFilePathName, InBaseData.SourceFrameNumberRenderThread);
+	ImageTask->Filename = FString::Printf(TEXT("%s%05d"), *BaseFilePathName, InBaseData.SourceFrameNumberRenderThread);
 	ImageTask->bOverwriteFile = bOverwriteFile;
 	ImageTask->CompressionQuality = CompressionQuality;
 	ImageTask->OnCompleted = OnCompleteWrapper;
