Post merge fix

diff --git a/Engine/Plugins/Media/MediaIOFramework/Source/MediaIOCore/Private/MediaCapture.cpp b/Engine/Plugins/Media/MediaIOFramework/Source/MediaIOCore/Private/MediaCapture.cpp
--- a/Engine/Plugins/Media/MediaIOFramework/Source/MediaIOCore/Private/MediaCapture.cpp
+++ b/Engine/Plugins/Media/MediaIOFramework/Source/MediaIOCore/Private/MediaCapture.cpp
@@ -719,7 +719,7 @@ void UMediaCapture::Capture_RenderThread(FRHICommandListImmediate& RHICmdList,
 			TexCreate_RenderTargetable,
 			false);
 		TRefCountPtr<IPooledRenderTarget> ResampleTexturePooledRenderTarget;
-		GetRendererModule().RenderTargetPoolFindFreeElement(RHICmdList, OutputDesc, ResampleTexturePooledRenderTarget, TEXT("MediaCapture"));
+		GRenderTargetPool.FindFreeElement(RHICmdList, OutputDesc, ResampleTexturePooledRenderTarget, TEXT("MediaCapture"));
 		const FSceneRenderTargetItem& DestRenderTarget = ResampleTexturePooledRenderTarget->GetRenderTargetItem();
 
 		// Do we need to crop
