[Blackmagic]

- Updating the scheduling loop to avoid corrupting a frame that is currently being displayed
- Fixed a problem in the MediaCapture frames to be copied logic that was causing out of sequence frame numbers

diff --git a/Engine/Plugins/Media/MediaIOFramework/Source/MediaIOCore/Private/MediaCapture.cpp b/Engine/Plugins/Media/MediaIOFramework/Source/MediaIOCore/Private/MediaCapture.cpp
--- a/Engine/Plugins/Media/MediaIOFramework/Source/MediaIOCore/Private/MediaCapture.cpp
+++ b/Engine/Plugins/Media/MediaIOFramework/Source/MediaIOCore/Private/MediaCapture.cpp
@@ -101,6 +101,7 @@ UMediaCapture::UMediaCapture(const FObjectInitializer& ObjectInitializer)
 	, bResolvedTargetInitialized(false)
 	, bShouldCaptureRHITexture(false)
 	, WaitingForResolveCommandExecutionCounter(0)
+	, NumberOfTexturesToResolve(0)
 {
 }
 
@@ -584,14 +585,13 @@ void UMediaCapture::OnEndFrame_GameThread()
 	int32 ReadyFrameIndex = (CurrentResolvedTargetIndex + 1) % NumberOfCaptureFrame; // Next one in the buffer queue
 
 	// Frame that should be on the system ram and we want to send to the user
-	FCaptureFrame* ReadyFrame = (CaptureFrames[ReadyFrameIndex].bResolvedTargetRequested) ? &CaptureFrames[ReadyFrameIndex] : nullptr;
+	FCaptureFrame* ReadyFrame = &CaptureFrames[ReadyFrameIndex];
 	// Frame that we want to transfer to system ram
 	FCaptureFrame* CapturingFrame = (GetState() != EMediaCaptureState::StopRequested) ? &CaptureFrames[CurrentResolvedTargetIndex] : nullptr;
 
 	UE_LOG(LogMediaIOCore, VeryVerbose, TEXT("MediaOutput: '%s'. ReadyFrameIndex: '%d' '%s'. CurrentResolvedTargetIndex: '%d'.")
 		, *MediaOutputName, ReadyFrameIndex, (CaptureFrames[ReadyFrameIndex].bResolvedTargetRequested) ? TEXT("Y"): TEXT("N"), CurrentResolvedTargetIndex);
-
-	if (ReadyFrame == nullptr && GetState() == EMediaCaptureState::StopRequested)
+	if (GetState() == EMediaCaptureState::StopRequested && NumberOfTexturesToResolve.Load() <= 0)
 	{
 		// All the requested frames have been captured.
 		StopCapture(false);
@@ -868,29 +868,34 @@ void UMediaCapture::Capture_RenderThread(FRHICommandListImmediate& RHICmdList,
 				// Asynchronously copy duplicate target from GPU to System Memory
 				RHICmdList.CopyToResolveTarget(DestRenderTarget.TargetableTexture, CapturingFrame->ReadbackTexture, FResolveParams());
 				CapturingFrame->bResolvedTargetRequested = true;
+				++NumberOfTexturesToResolve;
 			}
 		}
 	}
 
-	if (!InMediaCapture->bShouldCaptureRHITexture && ReadyFrame && InMediaCapture->GetState() != EMediaCaptureState::Error)
+	if (!InMediaCapture->bShouldCaptureRHITexture && InMediaCapture->GetState() != EMediaCaptureState::Error)
 	{
-		check(ReadyFrame->ReadbackTexture.IsValid());
-
-		// Lock & read
-		void* ColorDataBuffer = nullptr;
-		int32 Width = 0, Height = 0;
+		if (ReadyFrame->bResolvedTargetRequested)
 		{
-			SCOPE_CYCLE_COUNTER(STAT_MediaCapture_RenderThread_MapStaging);
-			RHICmdList.MapStagingSurface(ReadyFrame->ReadbackTexture, ColorDataBuffer, Width, Height);
-		}
+			check(ReadyFrame->ReadbackTexture.IsValid());
 
-		{
-			SCOPE_CYCLE_COUNTER(STAT_MediaCapture_RenderThread_Callback);
-			InMediaCapture->OnFrameCaptured_RenderingThread(ReadyFrame->CaptureBaseData, ReadyFrame->UserData, ColorDataBuffer, Width, Height);
-		}
-		ReadyFrame->bResolvedTargetRequested = false;
+			// Lock & read
+			void* ColorDataBuffer = nullptr;
+			int32 Width = 0, Height = 0;
+			{
+				SCOPE_CYCLE_COUNTER(STAT_MediaCapture_RenderThread_MapStaging);
+				RHICmdList.MapStagingSurface(ReadyFrame->ReadbackTexture, ColorDataBuffer, Width, Height);
+			}
+
+			{
+				SCOPE_CYCLE_COUNTER(STAT_MediaCapture_RenderThread_Callback);
+				InMediaCapture->OnFrameCaptured_RenderingThread(ReadyFrame->CaptureBaseData, ReadyFrame->UserData, ColorDataBuffer, Width, Height);
+			}
+			ReadyFrame->bResolvedTargetRequested = false;
+			--NumberOfTexturesToResolve;
 
-		RHICmdList.UnmapStagingSurface(ReadyFrame->ReadbackTexture);
+			RHICmdList.UnmapStagingSurface(ReadyFrame->ReadbackTexture);
+		}
 	}
 
 	--InMediaCapture->WaitingForResolveCommandExecutionCounter;
diff --git a/Engine/Plugins/Media/MediaIOFramework/Source/MediaIOCore/Public/MediaCapture.h b/Engine/Plugins/Media/MediaIOFramework/Source/MediaIOCore/Public/MediaCapture.h
--- a/Engine/Plugins/Media/MediaIOFramework/Source/MediaIOCore/Public/MediaCapture.h
+++ b/Engine/Plugins/Media/MediaIOFramework/Source/MediaIOCore/Public/MediaCapture.h
@@ -322,4 +322,5 @@ private:
 	TAtomic<bool> bResolvedTargetInitialized;
 	TAtomic<bool> bShouldCaptureRHITexture;
 	TAtomic<int32> WaitingForResolveCommandExecutionCounter;
+	TAtomic<int32> NumberOfTexturesToResolve;
 };
