Fix InlineEditCondition for FileMediaOutput.

diff --git a/Engine/Plugins/Media/MediaIOFramework/Source/MediaIOCore/Public/FileMediaOutput.h b/Engine/Plugins/Media/MediaIOFramework/Source/MediaIOCore/Public/FileMediaOutput.h
--- a/Engine/Plugins/Media/MediaIOFramework/Source/MediaIOCore/Public/FileMediaOutput.h
+++ b/Engine/Plugins/Media/MediaIOFramework/Source/MediaIOCore/Public/FileMediaOutput.h
@@ -46,7 +46,7 @@ public:
 	FString BaseFileName;
 
 	/** Use the default back buffer size or specify a specific size to capture. */
-	UPROPERTY(BlueprintReadWrite, Category="Media")
+	UPROPERTY(BlueprintReadWrite, EditAnywhere, Category="Media", meta = (InlineEditConditionToggle))
 	bool bOverrideDesiredSize;
 
 	/** Use the default back buffer size or specify a specific size to capture. */
@@ -54,7 +54,7 @@ public:
 	FIntPoint DesiredSize;
 
 	/** Use the default back buffer pixel format or specify a specific the pixel format to capture. */
-	UPROPERTY(BlueprintReadWrite, Category="Media")
+	UPROPERTY(BlueprintReadWrite, EditAnywhere, Category="Media", meta = (InlineEditConditionToggle))
 	bool bOverridePixelFormat;
 
 	/** Use the default back buffer pixel format or specify a specific the pixel format to capture. */
