[MediaCapture] Prevent crash when the requested size is negative.

diff --git a/Engine/Plugins/Media/MediaIOFramework/Source/MediaIOCore/Private/MediaCapture.cpp b/Engine/Plugins/Media/MediaIOFramework/Source/MediaIOCore/Private/MediaCapture.cpp
--- a/Engine/Plugins/Media/MediaIOFramework/Source/MediaIOCore/Private/MediaCapture.cpp
+++ b/Engine/Plugins/Media/MediaIOFramework/Source/MediaIOCore/Private/MediaCapture.cpp
@@ -176,22 +176,30 @@ bool UMediaCapture::CaptureSceneViewport(TSharedPtr<FSceneViewport>& InSceneView
 	}
 
 	SetState(EMediaCaptureState::Preparing);
-	if (!CaptureSceneViewportImpl(InSceneViewport))
+	bool bInitialized = CaptureSceneViewportImpl(InSceneViewport);
+
+	if (bInitialized)
+	{
+		InitializeResolveTarget(MediaOutput->NumberOfTextureBuffers);
+		bInitialized = GetState() != EMediaCaptureState::Stopped;
+	}
+
+	if (bInitialized)
+	{
+		//no lock required, the command on the render thread is not active
+		CapturingSceneViewport = InSceneViewport;
+
+		CurrentResolvedTargetIndex = 0;
+		FCoreDelegates::OnEndFrame.AddUObject(this, &UMediaCapture::OnEndFrame_GameThread);
+	}
+	else
 	{
 		ResetFixedViewportSize(InSceneViewport, false);
 		SetState(EMediaCaptureState::Stopped);
 		MediaCaptureDetails::ShowSlateNotification();
-		return false;
 	}
 
-	//no lock required, the command on the render thread is not active
-	CapturingSceneViewport = InSceneViewport;
-
-	InitializeResolveTarget(MediaOutput->NumberOfTextureBuffers);
-	CurrentResolvedTargetIndex = 0;
-	FCoreDelegates::OnEndFrame.AddUObject(this, &UMediaCapture::OnEndFrame_GameThread);
-
-	return true;
+	return bInitialized;
 }
 
 bool UMediaCapture::CaptureTextureRenderTarget2D(UTextureRenderTarget2D* InRenderTarget2D, FMediaCaptureOptions CaptureOptions)
@@ -228,21 +236,29 @@ bool UMediaCapture::CaptureTextureRenderTarget2D(UTextureRenderTarget2D* InRende
 	}
 
 	SetState(EMediaCaptureState::Preparing);
-	if (!CaptureRenderTargetImpl(InRenderTarget2D))
+	bool bInitialized = CaptureRenderTargetImpl(InRenderTarget2D);
+
+	if (bInitialized)
 	{
-		SetState(EMediaCaptureState::Stopped);
-		MediaCaptureDetails::ShowSlateNotification();
-		return false;
+		InitializeResolveTarget(MediaOutput->NumberOfTextureBuffers);
+		bInitialized = GetState() != EMediaCaptureState::Stopped;
 	}
 
-	//no lock required the command on the render thread is not active yet
-	CapturingRenderTarget = InRenderTarget2D;
+	if (bInitialized)
+	{
+		//no lock required the command on the render thread is not active yet
+		CapturingRenderTarget = InRenderTarget2D;
 
-	InitializeResolveTarget(MediaOutput->NumberOfTextureBuffers);
-	CurrentResolvedTargetIndex = 0;
-	FCoreDelegates::OnEndFrame.AddUObject(this, &UMediaCapture::OnEndFrame_GameThread);
+		CurrentResolvedTargetIndex = 0;
+		FCoreDelegates::OnEndFrame.AddUObject(this, &UMediaCapture::OnEndFrame_GameThread);
+	}
+	else
+	{
+		SetState(EMediaCaptureState::Stopped);
+		MediaCaptureDetails::ShowSlateNotification();
+	}
 
-	return true;
+	return bInitialized;
 }
 
 void UMediaCapture::CacheMediaOutput(EMediaCaptureSourceType InSourceType)
@@ -502,6 +518,13 @@ bool UMediaCapture::HasFinishedProcessing() const
 
 void UMediaCapture::InitializeResolveTarget(int32 InNumberOfBuffers)
 {
+	if (DesiredOutputSize.X <= 0 || DesiredOutputSize.Y <= 0)
+	{
+		UE_LOG(LogMediaIOCore, Error, TEXT("Can't start the capture. The size requested is negative or zero."));
+		SetState(EMediaCaptureState::Stopped);
+		return;
+	}
+
 	if (bShouldCaptureRHITexture)
 	{
 		// No buffer is needed if the callback is with the RHI Texture
