- Fixing MediaOutput asset creation

diff --git a/Engine/Plugins/Media/MediaIOFramework/Source/MediaIOEditor/Private/AssetTypeActions/AssetTypeActions_MediaOutput.cpp b/Engine/Plugins/Media/MediaIOFramework/Source/MediaIOEditor/Private/AssetTypeActions/AssetTypeActions_MediaOutput.cpp
new file mode 100644
--- /dev/null
+++ b/Engine/Plugins/Media/MediaIOFramework/Source/MediaIOEditor/Private/AssetTypeActions/AssetTypeActions_MediaOutput.cpp
@@ -0,0 +1,34 @@
+// Copyright 1998-2019 Epic Games, Inc. All Rights Reserved.
+
+#include "AssetTypeActions/AssetTypeActions_MediaOutput.h"
+
+#include "MediaOutput.h"
+
+#define LOCTEXT_NAMESPACE "AssetTypeActions_MediaOutput"
+
+FColor FAssetTypeActions_MediaOutput::GetTypeColor() const
+{
+	return FColor::White;
+}
+
+uint32 FAssetTypeActions_MediaOutput::GetCategories()
+{
+	return EAssetTypeCategories::Media;
+}
+
+bool FAssetTypeActions_MediaOutput::IsImportedAsset() const
+{
+	return false;
+}
+
+FText FAssetTypeActions_MediaOutput::GetName() const
+{
+	return FText::GetEmpty(); //Factory will give the name for each asset type
+}
+
+UClass* FAssetTypeActions_MediaOutput::GetSupportedClass() const
+{
+	return UMediaOutput::StaticClass();
+}
+
+#undef LOCTEXT_NAMESPACE
diff --git a/Engine/Plugins/Media/MediaIOFramework/Source/MediaIOEditor/Private/AssetTypeActions/AssetTypeActions_MediaOutput.h b/Engine/Plugins/Media/MediaIOFramework/Source/MediaIOEditor/Private/AssetTypeActions/AssetTypeActions_MediaOutput.h
new file mode 100644
--- /dev/null
+++ b/Engine/Plugins/Media/MediaIOFramework/Source/MediaIOEditor/Private/AssetTypeActions/AssetTypeActions_MediaOutput.h
@@ -0,0 +1,18 @@
+// Copyright 1998-2019 Epic Games, Inc. All Rights Reserved.
+
+#pragma once
+
+#include "CoreMinimal.h"
+#include "AssetTypeActions_Base.h"
+
+class FAssetTypeActions_MediaOutput : public FAssetTypeActions_Base
+{
+public:
+	//~ Begin IAssetTypeActions Implementation
+	virtual FColor GetTypeColor() const override;
+	virtual uint32 GetCategories() override;
+	virtual bool IsImportedAsset() const override;
+	virtual FText GetName() const override;
+	virtual UClass* GetSupportedClass() const override;
+	//~ End IAssetTypeActions Implementation
+};
\ No newline at end of file
diff --git a/Engine/Plugins/Media/MediaIOFramework/Source/MediaIOEditor/Private/MediaIOEditorModule.cpp b/Engine/Plugins/Media/MediaIOFramework/Source/MediaIOEditor/Private/MediaIOEditorModule.cpp
--- a/Engine/Plugins/Media/MediaIOFramework/Source/MediaIOEditor/Private/MediaIOEditorModule.cpp
+++ b/Engine/Plugins/Media/MediaIOFramework/Source/MediaIOEditor/Private/MediaIOEditorModule.cpp
@@ -12,6 +12,8 @@
 #include "Styling/SlateStyle.h"
 #include "Styling/SlateStyleRegistry.h"
 
+#include "AssetTypeActions/AssetTypeActions_MediaOutput.h"
+
 #include "Customizations/MediaIODeviceCustomization.h"
 #include "Customizations/MediaIOConfigurationCustomization.h"
 #include "Customizations/MediaIOInputConfigurationCustomization.h"
@@ -26,6 +28,7 @@ class FMediaIOEditorModule : public IModuleInterface
 
 	virtual void StartupModule() override
 	{
+		RegisterAssetTools();
 		RegisterCustomizations();
 		RegisterStyle();
 	}
@@ -36,10 +39,45 @@ class FMediaIOEditorModule : public IModuleInterface
 		{
 			UnregisterStyle();
 			UnregisterCustomizations();
+			UnregisterAssetTools();
 		}
 	}
 
 private:
+	/** Registers asset tool actions. */
+	void RegisterAssetTools()
+	{
+		IAssetTools& AssetTools = FModuleManager::LoadModuleChecked<FAssetToolsModule>("AssetTools").Get();
+		RegisterAssetTypeAction(AssetTools, MakeShareable(new FAssetTypeActions_MediaOutput));
+	}
+
+	/**
+	 * Registers a single asset type action.
+	 */
+	void RegisterAssetTypeAction(IAssetTools& AssetTools, TSharedRef<IAssetTypeActions> Action)
+	{
+		AssetTools.RegisterAssetTypeActions(Action);
+		RegisteredAssetTypeActions.Add(Action);
+	}
+
+	/** Unregisters asset tool actions. */
+	void UnregisterAssetTools()
+	{
+		FAssetToolsModule* AssetToolsModule = FModuleManager::GetModulePtr<FAssetToolsModule>("AssetTools");
+
+		if (AssetToolsModule != nullptr)
+		{
+			IAssetTools& AssetTools = AssetToolsModule->Get();
+
+			for (const TSharedRef<IAssetTypeActions>& Action : RegisteredAssetTypeActions)
+			{
+				AssetTools.UnregisterAssetTypeActions(Action);
+			}
+		}
+
+		RegisteredAssetTypeActions.Reset();
+	}
+
 	/** Register details view customizations. */
 	void RegisterCustomizations()
 	{
@@ -81,6 +119,10 @@ private:
 	{
 		FSlateStyleRegistry::UnRegisterSlateStyle(*StyleInstance.Get());
 	}
+
+private:
+	/** The collection of registered asset type actions. */
+	TArray<TSharedRef<IAssetTypeActions>> RegisteredAssetTypeActions;
 };
 
 IMPLEMENT_MODULE(FMediaIOEditorModule, MediaIOEditor)
