Fix copyright notice

#rnx

diff --git a/Engine/Plugins/Animation/LiveLink/Source/LiveLinkEditor/Private/LiveLinkSubjectKeyDetailCustomization.cpp b/Engine/Plugins/Animation/LiveLink/Source/LiveLinkEditor/Private/LiveLinkSubjectKeyDetailCustomization.cpp
--- a/Engine/Plugins/Animation/LiveLink/Source/LiveLinkEditor/Private/LiveLinkSubjectKeyDetailCustomization.cpp
+++ b/Engine/Plugins/Animation/LiveLink/Source/LiveLinkEditor/Private/LiveLinkSubjectKeyDetailCustomization.cpp
@@ -1,4 +1,4 @@
-// Copyright 1998-2019 Epic Games, Inc. All Rights Reserved.
+// Copyright Epic Games, Inc. All Rights Reserved.
 
 #include "LiveLinkSubjectKeyDetailCustomization.h"
 
diff --git a/Engine/Plugins/Animation/LiveLink/Source/LiveLinkEditor/Private/LiveLinkSubjectKeyDetailCustomization.h b/Engine/Plugins/Animation/LiveLink/Source/LiveLinkEditor/Private/LiveLinkSubjectKeyDetailCustomization.h
--- a/Engine/Plugins/Animation/LiveLink/Source/LiveLinkEditor/Private/LiveLinkSubjectKeyDetailCustomization.h
+++ b/Engine/Plugins/Animation/LiveLink/Source/LiveLinkEditor/Private/LiveLinkSubjectKeyDetailCustomization.h
@@ -1,4 +1,4 @@
-// Copyright 1998-2019 Epic Games, Inc. All Rights Reserved.
+// Copyright Epic Games, Inc. All Rights Reserved.
 
 #pragma once
 
diff --git a/Engine/Plugins/Importers/USDImporter/Source/USDUtilities/Private/USDErrorUtils.cpp b/Engine/Plugins/Importers/USDImporter/Source/USDUtilities/Private/USDErrorUtils.cpp
--- a/Engine/Plugins/Importers/USDImporter/Source/USDUtilities/Private/USDErrorUtils.cpp
+++ b/Engine/Plugins/Importers/USDImporter/Source/USDUtilities/Private/USDErrorUtils.cpp
@@ -1,4 +1,4 @@
-// Copyright 1998-2019 Epic Games, Inc. All Rights Reserved.
+// Copyright Epic Games, Inc. All Rights Reserved.
 
 #include "USDErrorUtils.h"
 
diff --git a/Engine/Plugins/Importers/USDImporter/Source/USDUtilities/Public/USDErrorUtils.h b/Engine/Plugins/Importers/USDImporter/Source/USDUtilities/Public/USDErrorUtils.h
--- a/Engine/Plugins/Importers/USDImporter/Source/USDUtilities/Public/USDErrorUtils.h
+++ b/Engine/Plugins/Importers/USDImporter/Source/USDUtilities/Public/USDErrorUtils.h
@@ -1,4 +1,4 @@
-// Copyright 1998-2019 Epic Games, Inc. All Rights Reserved.
+// Copyright Epic Games, Inc. All Rights Reserved.
 
 #pragma once
 
diff --git a/Engine/Plugins/Media/MediaIOFramework/Source/MediaIOEditor/Private/AssetTypeActions/AssetTypeActions_MediaOutput.cpp b/Engine/Plugins/Media/MediaIOFramework/Source/MediaIOEditor/Private/AssetTypeActions/AssetTypeActions_MediaOutput.cpp
--- a/Engine/Plugins/Media/MediaIOFramework/Source/MediaIOEditor/Private/AssetTypeActions/AssetTypeActions_MediaOutput.cpp
+++ b/Engine/Plugins/Media/MediaIOFramework/Source/MediaIOEditor/Private/AssetTypeActions/AssetTypeActions_MediaOutput.cpp
@@ -1,4 +1,4 @@
-// Copyright 1998-2019 Epic Games, Inc. All Rights Reserved.
+// Copyright Epic Games, Inc. All Rights Reserved.
 
 #include "AssetTypeActions/AssetTypeActions_MediaOutput.h"
 
diff --git a/Engine/Plugins/Media/MediaIOFramework/Source/MediaIOEditor/Private/AssetTypeActions/AssetTypeActions_MediaOutput.h b/Engine/Plugins/Media/MediaIOFramework/Source/MediaIOEditor/Private/AssetTypeActions/AssetTypeActions_MediaOutput.h
--- a/Engine/Plugins/Media/MediaIOFramework/Source/MediaIOEditor/Private/AssetTypeActions/AssetTypeActions_MediaOutput.h
+++ b/Engine/Plugins/Media/MediaIOFramework/Source/MediaIOEditor/Private/AssetTypeActions/AssetTypeActions_MediaOutput.h
@@ -1,4 +1,4 @@
-// Copyright 1998-2019 Epic Games, Inc. All Rights Reserved.
+// Copyright Epic Games, Inc. All Rights Reserved.
 
 #pragma once
 
diff --git a/Engine/Source/Runtime/Core/Private/Misc/Timecode.cpp b/Engine/Source/Runtime/Core/Private/Misc/Timecode.cpp
--- a/Engine/Source/Runtime/Core/Private/Misc/Timecode.cpp
+++ b/Engine/Source/Runtime/Core/Private/Misc/Timecode.cpp
@@ -1,4 +1,4 @@
-// Copyright 1998-2019 Epic Games, Inc. All Rights Reserved.
+// Copyright Epic Games, Inc. All Rights Reserved.
 
 #include "Misc/Timecode.h"
 #include "HAL/IConsoleManager.h"
