Merging CL #11569656, 11570593, 11575492 from Dev-VirtualProduction

- Fixed Aja/Blackmagic timed data samples types to use what is being captured
- Fixed potential crash in MediaIO sample caching
- Fixed MediaIO channels being linked to an old input and not being tracked again (when updating media profile)
- Fixed MediaIOPlayer not using the right timecode when setting its playtime

diff --git a/Engine/Plugins/Media/AjaMedia/Source/AjaMedia/Private/Player/AjaMediaPlayer.cpp b/Engine/Plugins/Media/AjaMedia/Source/AjaMedia/Private/Player/AjaMediaPlayer.cpp
--- a/Engine/Plugins/Media/AjaMedia/Source/AjaMedia/Private/Player/AjaMediaPlayer.cpp
+++ b/Engine/Plugins/Media/AjaMedia/Source/AjaMedia/Private/Player/AjaMediaPlayer.cpp
@@ -68,9 +68,8 @@ FAjaMediaPlayer::FAjaMediaPlayer(IMediaEventSink& InEventSink)
 	, bUseVideo(false)
 	, bVerifyFrameDropCount(true)
 	, InputChannel(nullptr)
-	, SupportedSampleTypes(static_cast<EMediaIOSampleType>(static_cast<uint8>(EMediaIOSampleType::Video) | static_cast<uint8>(EMediaIOSampleType::Audio) | static_cast<uint8>(EMediaIOSampleType::Metadata)))
+	, SupportedSampleTypes(EMediaIOSampleType::None)
 {
-	Samples->EnableTimedDataChannels(SupportedSampleTypes);
 }
 
 
@@ -190,6 +189,12 @@ bool FAjaMediaPlayer::Open(const FString& Url, const IMediaOptions* Options)
 		AjaOptions.bUseKey = false;
 		AjaOptions.bBurnTimecode = false;
 		AjaOptions.BurnTimecodePercentY = 80;
+
+		//Adjust supported sample types based on what's being captured
+		SupportedSampleTypes = AjaOptions.bUseVideo ? EMediaIOSampleType::Video : EMediaIOSampleType::None;
+		SupportedSampleTypes |= AjaOptions.bUseAudio ? EMediaIOSampleType::Audio : EMediaIOSampleType::None;
+		SupportedSampleTypes |= AjaOptions.bUseAncillary ? EMediaIOSampleType::Metadata : EMediaIOSampleType::None;
+		Samples->EnableTimedDataChannels(this, SupportedSampleTypes);
 	}
 
 	bVerifyFrameDropCount = Options->GetMediaOption(AjaMediaOption::LogDropFrame, true);
@@ -240,6 +245,9 @@ void FAjaMediaPlayer::Close()
 	MetadataSamplePool->Reset();
 	TextureSamplePool->Reset();
 
+	//Disable all our channels from the monitor
+	Samples->EnableTimedDataChannels(this, EMediaIOSampleType::None);
+
 	AjaThreadCurrentAncSample.Reset();
 	AjaThreadCurrentAncF2Sample.Reset();
 	AjaThreadCurrentAudioSample.Reset();
diff --git a/Engine/Plugins/Media/BlackmagicMedia/Source/BlackmagicMedia/Private/Player/BlackmagicMediaPlayer.cpp b/Engine/Plugins/Media/BlackmagicMedia/Source/BlackmagicMedia/Private/Player/BlackmagicMediaPlayer.cpp
--- a/Engine/Plugins/Media/BlackmagicMedia/Source/BlackmagicMedia/Private/Player/BlackmagicMediaPlayer.cpp
+++ b/Engine/Plugins/Media/BlackmagicMedia/Source/BlackmagicMedia/Private/Player/BlackmagicMediaPlayer.cpp
@@ -399,9 +399,8 @@ FBlackmagicMediaPlayer::FBlackmagicMediaPlayer(IMediaEventSink& InEventSink)
 	, AudioSamplePool(new FBlackmagicMediaAudioSamplePool)
 	, TextureSamplePool(new FBlackmagicMediaTextureSamplePool)
 	, bVerifyFrameDropCount(false)
-	, SupportedSampleTypes(static_cast<EMediaIOSampleType>(static_cast<uint8>(EMediaIOSampleType::Video) | static_cast<uint8>(EMediaIOSampleType::Audio)))
+	, SupportedSampleTypes(EMediaIOSampleType::None)
 {
-	Samples->EnableTimedDataChannels(SupportedSampleTypes);
 }
 
 FBlackmagicMediaPlayer::~FBlackmagicMediaPlayer()
@@ -425,6 +424,9 @@ void FBlackmagicMediaPlayer::Close()
 	AudioSamplePool->Reset();
 	TextureSamplePool->Reset();
 
+	//Disable all our channels from the monitor
+	Samples->EnableTimedDataChannels(this, EMediaIOSampleType::None);
+
 	Super::Close();
 }
 
@@ -489,6 +491,11 @@ bool FBlackmagicMediaPlayer::Open(const FString& Url, const IMediaOptions* Optio
 		ChannelOptions.NumberOfAudioChannel = (AudioChannelOption == EBlackmagicMediaAudioChannel::Surround8) ? 8 : 2;
 	}
 
+	//Adjust supported sample types based on what's being captured
+	SupportedSampleTypes = ChannelOptions.bReadVideo ? EMediaIOSampleType::Video : EMediaIOSampleType::None;
+	SupportedSampleTypes |= ChannelOptions.bReadAudio ? EMediaIOSampleType::Audio : EMediaIOSampleType::None;
+	Samples->EnableTimedDataChannels(this, SupportedSampleTypes);
+
 	bVerifyFrameDropCount = Options->GetMediaOption(BlackmagicMediaOption::LogDropFrame, false);
 	const bool bEncodeTimecodeInTexel = TimecodeFormat != EMediaIOTimecodeFormat::None && Options->GetMediaOption(BlackmagicMediaOption::EncodeTimecodeInTexel, false);
 	MaxNumAudioFrameBuffer = Options->GetMediaOption(BlackmagicMediaOption::MaxAudioFrameBuffer, (int64)8);
diff --git a/Engine/Plugins/Media/MediaIOFramework/Source/MediaIOCore/Private/Player/MediaIOCorePlayerBase.cpp b/Engine/Plugins/Media/MediaIOFramework/Source/MediaIOCore/Private/Player/MediaIOCorePlayerBase.cpp
--- a/Engine/Plugins/Media/MediaIOFramework/Source/MediaIOCore/Private/Player/MediaIOCorePlayerBase.cpp
+++ b/Engine/Plugins/Media/MediaIOFramework/Source/MediaIOCore/Private/Player/MediaIOCorePlayerBase.cpp
@@ -78,7 +78,7 @@ FMediaIOCorePlayerBase::FMediaIOCorePlayerBase(IMediaEventSink& InEventSink)
 	, TimeDelay(0.0)
 	, PreviousFrameTimespan(FTimespan::Zero())
 {
-	Samples = new FMediaIOCoreSamples(this);
+	Samples = new FMediaIOCoreSamples();
 }
 
 FMediaIOCorePlayerBase::~FMediaIOCorePlayerBase()
@@ -193,26 +193,31 @@ void FMediaIOCorePlayerBase::TickTimeManagement()
 {
 	if (bUseTimeSynchronization)
 	{
-		FTimecode Timecode = FApp::GetTimecode();
-		const FFrameRate FrameRate = FApp::GetTimecodeFrameRate();
-
-		if (!bWarnedIncompatibleFrameRate)
+		TOptional<FQualifiedFrameTime> CurrentFrameTime = FApp::GetCurrentFrameTime();
+		if (CurrentFrameTime.IsSet())
 		{
-			if (!FrameRate.IsMultipleOf(VideoFrameRate) && !VideoFrameRate.IsMultipleOf(FrameRate))
+			if (!bWarnedIncompatibleFrameRate)
 			{
-				bWarnedIncompatibleFrameRate = true;
-				UE_LOG(LogMediaIOCore, Warning, TEXT("The video's frame rate is incompatible with engine's frame rate. %s"), *OpenUrl);
+				if (!CurrentFrameTime->Rate.IsMultipleOf(VideoFrameRate) && !VideoFrameRate.IsMultipleOf(CurrentFrameTime->Rate))
+				{
+					bWarnedIncompatibleFrameRate = true;
+					UE_LOG(LogMediaIOCore, Warning, TEXT("The video's frame rate is incompatible with engine's frame rate. %s"), *OpenUrl);
+				}
 			}
-		}
 
-		if (FrameDelay != 0)
+			if (FrameDelay != 0)
+			{
+				FFrameTime FrameNumberVideoRate = CurrentFrameTime->ConvertTo(VideoFrameRate);
+				FrameNumberVideoRate.FrameNumber -= FrameDelay;
+				CurrentFrameTime->Time = FFrameRate::TransformTime(FrameNumberVideoRate, VideoFrameRate, CurrentFrameTime->Rate);
+			}
+
+			CurrentTime = FTimespan::FromSeconds(CurrentFrameTime->AsSeconds());
+		}
+		else
 		{
-			FFrameNumber FrameNumber = Timecode.ToFrameNumber(VideoFrameRate);
-			FrameNumber -= FrameDelay;
-			Timecode = FTimecode::FromFrameNumber(FrameNumber, FrameRate, Timecode.bDropFrameFormat);
+			UE_LOG(LogMediaIOCore, Verbose, TEXT("The video '%s' is configured to use timecode but none is available on the engine."), *OpenUrl);
 		}
-
-		CurrentTime = Timecode.ToTimespan(FrameRate);
 	}
 	else
 	{
diff --git a/Engine/Plugins/Media/MediaIOFramework/Source/MediaIOCore/Private/Player/MediaIOCoreSamples.cpp b/Engine/Plugins/Media/MediaIOFramework/Source/MediaIOCore/Private/Player/MediaIOCoreSamples.cpp
--- a/Engine/Plugins/Media/MediaIOFramework/Source/MediaIOCore/Private/Player/MediaIOCoreSamples.cpp
+++ b/Engine/Plugins/Media/MediaIOFramework/Source/MediaIOCore/Private/Player/MediaIOCoreSamples.cpp
@@ -8,12 +8,12 @@
 #include "IMediaTextureSample.h"
 
 
-FMediaIOCoreSamples::FMediaIOCoreSamples(ITimedDataInput* ChannelsOwner)
-	: VideoSamples(ChannelsOwner, "Video")
-	, AudioSamples(ChannelsOwner, "Audio")
-	, MetadataSamples(ChannelsOwner, "Metadata")
-	, SubtitleSamples(ChannelsOwner, "Subtitles")
-	, CaptionSamples(ChannelsOwner, "Caption")
+FMediaIOCoreSamples::FMediaIOCoreSamples()
+	: VideoSamples("Video")
+	, AudioSamples("Audio")
+	, MetadataSamples("Metadata")
+	, SubtitleSamples("Subtitles")
+	, CaptionSamples("Caption")
 {
 
 }
@@ -27,13 +27,13 @@ void FMediaIOCoreSamples::CacheSamplesState(FTimespan PlayerTime)
 	CaptionSamples.CacheState(PlayerTime);
 }
 
-void FMediaIOCoreSamples::EnableTimedDataChannels(EMediaIOSampleType SampleTypes)
+void FMediaIOCoreSamples::EnableTimedDataChannels(ITimedDataInput* Input, EMediaIOSampleType SampleTypes)
 {
-	VideoSamples.EnableChannel(static_cast<uint8>(SampleTypes) & static_cast<uint8>(EMediaIOSampleType::Video));
-	AudioSamples.EnableChannel(static_cast<uint8>(SampleTypes)& static_cast<uint8>(EMediaIOSampleType::Audio));
-	MetadataSamples.EnableChannel(static_cast<uint8>(SampleTypes)& static_cast<uint8>(EMediaIOSampleType::Metadata));
-	SubtitleSamples.EnableChannel(static_cast<uint8>(SampleTypes)& static_cast<uint8>(EMediaIOSampleType::Subtitles));
-	CaptionSamples.EnableChannel(static_cast<uint8>(SampleTypes)& static_cast<uint8>(EMediaIOSampleType::Caption));
+	VideoSamples.EnableChannel(Input,  EnumHasAnyFlags(SampleTypes, EMediaIOSampleType::Video));
+	AudioSamples.EnableChannel(Input, EnumHasAnyFlags(SampleTypes, EMediaIOSampleType::Audio));
+	MetadataSamples.EnableChannel(Input, EnumHasAnyFlags(SampleTypes, EMediaIOSampleType::Metadata));
+	SubtitleSamples.EnableChannel(Input, EnumHasAnyFlags(SampleTypes, EMediaIOSampleType::Subtitles));
+	CaptionSamples.EnableChannel(Input, EnumHasAnyFlags(SampleTypes, EMediaIOSampleType::Caption));
 }
 
 void FMediaIOCoreSamples::InitializeVideoBuffer(const FMediaIOSamplingSettings& InSettings)
diff --git a/Engine/Plugins/Media/MediaIOFramework/Source/MediaIOCore/Public/MediaIOCoreSampleContainer.h b/Engine/Plugins/Media/MediaIOFramework/Source/MediaIOCore/Public/MediaIOCoreSampleContainer.h
--- a/Engine/Plugins/Media/MediaIOFramework/Source/MediaIOCore/Public/MediaIOCoreSampleContainer.h
+++ b/Engine/Plugins/Media/MediaIOFramework/Source/MediaIOCore/Public/MediaIOCoreSampleContainer.h
@@ -34,9 +34,8 @@ template<typename SampleType>
 class MEDIAIOCORE_API FMediaIOCoreSampleContainer : public ITimedDataInputChannel
 {
 public:
-	FMediaIOCoreSampleContainer(ITimedDataInput* Owner, FName InChannelName)
-		: Input(Owner)
-		, ChannelName(InChannelName)
+	FMediaIOCoreSampleContainer(FName InChannelName)
+		: ChannelName(InChannelName)
 		, BufferUnderflow(0)
 		, BufferOverflow(0)
 		, FrameDrop(0)
@@ -93,8 +92,15 @@ public:
 				if (EvaluationSettings.EvaluationType == ETimedDataInputEvaluationType::Timecode)
 				{
 					//Compute the distance with Timespan resolution. Going through FQualifiedFrameTime gives a different result (~10ns)
-					NewestSampleInSeconds = Samples[0]->GetTimecode().GetValue().ToTimespan(EvaluationSettings.FrameRate).GetTotalSeconds() + Duration;
-					OldestSampleInSeconds = Samples[CachedSamplesData.Num() - 1]->GetTimecode().GetValue().ToTimespan(EvaluationSettings.FrameRate).GetTotalSeconds();
+					if (Samples[0]->GetTimecode().IsSet())
+					{
+						NewestSampleInSeconds = Samples[0]->GetTimecode().GetValue().ToTimespan(EvaluationSettings.FrameRate).GetTotalSeconds() + Duration;
+					}
+
+					if (Samples[CachedSamplesData.Num() - 1]->GetTimecode().IsSet())
+					{
+						OldestSampleInSeconds = Samples[CachedSamplesData.Num() - 1]->GetTimecode().GetValue().ToTimespan(EvaluationSettings.FrameRate).GetTotalSeconds();
+					}
 				}
 				else //Platform time
 				{
@@ -126,27 +132,23 @@ public:
 	}
 
 	/** Channel is disabled by default. It won't be added to the Timed Data collection if not enabled */
-	void EnableChannel(bool bShouldEnable)
+	void EnableChannel(ITimedDataInput* Input, bool bShouldEnable)
 	{
+		check(Input);
+
 		if (bShouldEnable != bIsChannelEnabled)
 		{
 			bIsChannelEnabled = bShouldEnable;
 
 			if (bIsChannelEnabled)
 			{
-				if (Input)
-				{
-					Input->AddChannel(this);
-				}
+				Input->AddChannel(this);
 				ITimeManagementModule::Get().GetTimedDataInputCollection().Add(this);
 			}
 			else
 			{
 				ITimeManagementModule::Get().GetTimedDataInputCollection().Remove(this);
-				if (Input)
-				{
-					Input->RemoveChannel(this);
-				}
+				Input->RemoveChannel(this);
 			}
 		}
 	}
@@ -375,9 +377,6 @@ public:
 
 protected:
 
-	/** Input in which that channel belongs */
-	ITimedDataInput* Input;
-
 	/** Name of this channel */
 	FName ChannelName;
 
diff --git a/Engine/Plugins/Media/MediaIOFramework/Source/MediaIOCore/Public/MediaIOCoreSamples.h b/Engine/Plugins/Media/MediaIOFramework/Source/MediaIOCore/Public/MediaIOCoreSamples.h
--- a/Engine/Plugins/Media/MediaIOFramework/Source/MediaIOCore/Public/MediaIOCoreSamples.h
+++ b/Engine/Plugins/Media/MediaIOFramework/Source/MediaIOCore/Public/MediaIOCoreSamples.h
@@ -17,7 +17,6 @@ class IMediaTextureSample;
 /**
  * Bitflags about different types of sample a MediaIO can push
  */
-UENUM()
 enum class EMediaIOSampleType
 {
 	None = 0,
@@ -27,6 +26,7 @@ enum class EMediaIOSampleType
 	Subtitles = 1<<3,
 	Caption = 1<<4,
 };
+ENUM_CLASS_FLAGS(EMediaIOSampleType)
 
 /**
  * General purpose media sample queue.
@@ -35,7 +35,7 @@ enum class EMediaIOSampleType
 class MEDIAIOCORE_API FMediaIOCoreSamples : public IMediaSamples
 {
 public:
-	FMediaIOCoreSamples(ITimedDataInput* ChannelsOwner);
+	FMediaIOCoreSamples();
 	FMediaIOCoreSamples(const FMediaIOCoreSamples&) = delete;
 	FMediaIOCoreSamples& operator=(const FMediaIOCoreSamples&) = delete;
 
@@ -278,7 +278,7 @@ public:
 	void CacheSamplesState(FTimespan PlayerTime);
 
 	/** Enable or disable channels based on bitfield flag */
-	void EnableTimedDataChannels(EMediaIOSampleType SampleTypes);
+	void EnableTimedDataChannels(ITimedDataInput* Input, EMediaIOSampleType SampleTypes);
 
 	/** Initialize our different buffers with player's settings and wheter or not it should be displayed / supported in Timing monitor */
 	void InitializeVideoBuffer(const FMediaIOSamplingSettings& InSettings);
diff --git a/Engine/Plugins/VirtualProduction/TimedDataMonitor/Source/TimedDataMonitor/Private/TimedDataMonitorSubsystem.cpp b/Engine/Plugins/VirtualProduction/TimedDataMonitor/Source/TimedDataMonitor/Private/TimedDataMonitorSubsystem.cpp
--- a/Engine/Plugins/VirtualProduction/TimedDataMonitor/Source/TimedDataMonitor/Private/TimedDataMonitorSubsystem.cpp
+++ b/Engine/Plugins/VirtualProduction/TimedDataMonitor/Source/TimedDataMonitor/Private/TimedDataMonitorSubsystem.cpp
@@ -3,6 +3,7 @@
 #include "TimedDataMonitorSubsystem.h"
 
 #include "Engine/Engine.h"
+#include "Engine/TimecodeProvider.h"
 #include "Engine/World.h"
 #include "ITimeManagementModule.h"
 #include "LatentActions.h"
@@ -13,6 +14,10 @@
 #include "TimedDataMonitorCalibration.h"
 #include "UObject/Stack.h"
 
+#if !(UE_BUILD_SHIPPING || UE_BUILD_TEST)
+#include "Misc/FileHelper.h"
+#endif //!(UE_BUILD_SHIPPING || UE_BUILD_TEST)
+
 
 static TAutoConsoleVariable<bool> CVarEnableTimedDataMonitorSubsystemStats(TEXT("TimedDataMonitor.EnableStatUpdate"), 1, TEXT("Enable calculating evaluation statistics of all registered channels."));
 
diff --git a/Engine/Plugins/VirtualProduction/TimedDataMonitor/Source/TimedDataMonitorEditor/Private/STimingDiagramWidget.cpp b/Engine/Plugins/VirtualProduction/TimedDataMonitor/Source/TimedDataMonitorEditor/Private/STimingDiagramWidget.cpp
--- a/Engine/Plugins/VirtualProduction/TimedDataMonitor/Source/TimedDataMonitorEditor/Private/STimingDiagramWidget.cpp
+++ b/Engine/Plugins/VirtualProduction/TimedDataMonitor/Source/TimedDataMonitorEditor/Private/STimingDiagramWidget.cpp
@@ -192,7 +192,8 @@ public:
 
 		// Draw evaluation line
 		{
-			FLinearColor EvaluationColor = (EvaluationTime >= MinSampleTime && EvaluationTime <= MaxSampleTime) ? FLinearColor::Green : FLinearColor::Red;
+			const bool bIsInRange = (FMath::IsNearlyEqual(EvaluationTime, MinSampleTime) || EvaluationTime >= MinSampleTime) && (FMath::IsNearlyEqual(EvaluationTime, MaxSampleTime) || EvaluationTime <= MaxSampleTime);
+			FLinearColor EvaluationColor = bIsInRange ? FLinearColor::Green : FLinearColor::Red;
 			FSlateDrawElement::MakeBox(OutDrawElements, LayerId,
 				AllottedGeometry.ToPaintGeometry(FVector2D(LocationOfCenter, 0), FVector2D(1, SizeY)),
 				DarkBrush, ESlateDrawEffect::None,
@@ -347,4 +348,4 @@ FText STimingDiagramWidget::GetTooltipText() const
 
 }
 
-#undef LOCTEXT_NAMESPACE
\ No newline at end of file
+#undef LOCTEXT_NAMESPACE
