- Fixing empty filter entry in Media category

diff --git a/Engine/Plugins/Media/MediaIOFramework/Source/MediaIOEditor/Private/AssetTypeActions/AssetTypeActions_MediaOutput.cpp b/Engine/Plugins/Media/MediaIOFramework/Source/MediaIOEditor/Private/AssetTypeActions/AssetTypeActions_MediaOutput.cpp
--- a/Engine/Plugins/Media/MediaIOFramework/Source/MediaIOEditor/Private/AssetTypeActions/AssetTypeActions_MediaOutput.cpp
+++ b/Engine/Plugins/Media/MediaIOFramework/Source/MediaIOEditor/Private/AssetTypeActions/AssetTypeActions_MediaOutput.cpp
@@ -6,6 +6,11 @@
 
 #define LOCTEXT_NAMESPACE "AssetTypeActions_MediaOutput"
 
+bool FAssetTypeActions_MediaOutput::CanFilter()
+{
+	return false;
+}
+
 FColor FAssetTypeActions_MediaOutput::GetTypeColor() const
 {
 	return FColor::White;
diff --git a/Engine/Plugins/Media/MediaIOFramework/Source/MediaIOEditor/Private/AssetTypeActions/AssetTypeActions_MediaOutput.h b/Engine/Plugins/Media/MediaIOFramework/Source/MediaIOEditor/Private/AssetTypeActions/AssetTypeActions_MediaOutput.h
--- a/Engine/Plugins/Media/MediaIOFramework/Source/MediaIOEditor/Private/AssetTypeActions/AssetTypeActions_MediaOutput.h
+++ b/Engine/Plugins/Media/MediaIOFramework/Source/MediaIOEditor/Private/AssetTypeActions/AssetTypeActions_MediaOutput.h
@@ -9,6 +9,7 @@ class FAssetTypeActions_MediaOutput : public FAssetTypeActions_Base
 {
 public:
 	//~ Begin IAssetTypeActions Implementation
+	virtual bool CanFilter() override;
 	virtual FColor GetTypeColor() const override;
 	virtual uint32 GetCategories() override;
 	virtual bool IsImportedAsset() const override;
diff --git a/Engine/Source/Editor/ContentBrowser/Private/SFilterList.cpp b/Engine/Source/Editor/ContentBrowser/Private/SFilterList.cpp
--- a/Engine/Source/Editor/ContentBrowser/Private/SFilterList.cpp
+++ b/Engine/Source/Editor/ContentBrowser/Private/SFilterList.cpp
@@ -1201,7 +1201,7 @@ void SFilterList::CreateFiltersMenuCategory(FMenuBuilder& MenuBuilder, const TAr
 		if ( WeakTypeActions.IsValid() )
 		{
 			TSharedPtr<IAssetTypeActions> TypeActions = WeakTypeActions.Pin();
-			if ( TypeActions.IsValid() )
+			if ( TypeActions.IsValid() && TypeActions->CanFilter() )
 			{
 				const FText& LabelText = TypeActions->GetName();
 				MenuBuilder.AddMenuEntry(
