MediaIO: Fixing invalid test condition in media output

diff --git a/Engine/Plugins/Media/MediaIOFramework/Source/MediaIOCore/Private/MediaOutput.cpp b/Engine/Plugins/Media/MediaIOFramework/Source/MediaIOCore/Private/MediaOutput.cpp
--- a/Engine/Plugins/Media/MediaIOFramework/Source/MediaIOCore/Private/MediaOutput.cpp
+++ b/Engine/Plugins/Media/MediaIOFramework/Source/MediaIOCore/Private/MediaOutput.cpp
@@ -46,7 +46,7 @@ bool UMediaOutput::Validate(FString& OutFailureReason) const
 	}
 
 	const int32 MaxSupportedNumberOfbuffers = 8; // Arbitrary number
-	if (NumberOfTextureBuffers < 1 && NumberOfTextureBuffers > MaxSupportedNumberOfbuffers)
+	if (NumberOfTextureBuffers < 1 || NumberOfTextureBuffers > MaxSupportedNumberOfbuffers)
 	{
 		OutFailureReason = TEXT("NumberOfTextureBuffers is not valid.");
 		return false;
