MediaIOFramework: Added option to not invert alpha when writing out exr files.

#rnx

diff --git a/Engine/Plugins/Media/MediaIOFramework/Source/MediaIOCore/Private/FileMediaOutput.cpp b/Engine/Plugins/Media/MediaIOFramework/Source/MediaIOCore/Private/FileMediaOutput.cpp
--- a/Engine/Plugins/Media/MediaIOFramework/Source/MediaIOCore/Private/FileMediaOutput.cpp
+++ b/Engine/Plugins/Media/MediaIOFramework/Source/MediaIOCore/Private/FileMediaOutput.cpp
@@ -10,6 +10,7 @@
 
 UFileMediaOutput::UFileMediaOutput()
 	: Super()
+	, bInvertAlpha(true)
 {
 	FilePath.Path = FPaths::Combine(*FPaths::ProjectSavedDir(), *FString("MediaOutput"));
 }
@@ -86,7 +87,8 @@ EMediaCaptureConversionOperation UFileMediaOutput::GetConversionOperation(EMedia
 		EAlphaChannelMode::Type PropagateAlpha = EAlphaChannelMode::FromInt(CVarPropagateAlpha->GetValueOnAnyThread());
 		if ((PropagateAlpha == EAlphaChannelMode::AllowThroughTonemapper) || (InSourceType == EMediaCaptureSourceType::RENDER_TARGET))
 		{
-			return EMediaCaptureConversionOperation::INVERT_ALPHA;
+			return bInvertAlpha ? EMediaCaptureConversionOperation::INVERT_ALPHA :
+				EMediaCaptureConversionOperation::NONE;
 		}
 		else
 		{
diff --git a/Engine/Plugins/Media/MediaIOFramework/Source/MediaIOCore/Public/FileMediaOutput.h b/Engine/Plugins/Media/MediaIOFramework/Source/MediaIOCore/Public/FileMediaOutput.h
--- a/Engine/Plugins/Media/MediaIOFramework/Source/MediaIOCore/Public/FileMediaOutput.h
+++ b/Engine/Plugins/Media/MediaIOFramework/Source/MediaIOCore/Public/FileMediaOutput.h
@@ -61,6 +61,10 @@ public:
 	UPROPERTY(BlueprintReadWrite, EditAnywhere, Category="Media", meta=(EditCondition="bOverridePixelFormat"))
 	EFileMediaOutputPixelFormat DesiredPixelFormat;
 
+	/** Invert the alpha for formats that support alpha. */
+	UPROPERTY(BlueprintReadWrite, EditAnywhere, Category="Media")
+	bool bInvertAlpha;
+
 	//~ UMediaOutput interface
 public:
 	virtual bool Validate(FString& FailureReason) const override;
