Fix media IO audio output issues

diff --git a/Engine/Plugins/Media/AjaMedia/Source/AjaMediaOutput/Private/AjaMediaCapture.cpp b/Engine/Plugins/Media/AjaMedia/Source/AjaMediaOutput/Private/AjaMediaCapture.cpp
--- a/Engine/Plugins/Media/AjaMedia/Source/AjaMediaOutput/Private/AjaMediaCapture.cpp
+++ b/Engine/Plugins/Media/AjaMedia/Source/AjaMediaOutput/Private/AjaMediaCapture.cpp
@@ -291,6 +291,8 @@ void UAjaMediaCapture::StopCaptureImpl(bool bAllowPendingFrameToBeProcess)
 			}
 		}
 
+		AudioOutput.Reset();
+
 		RestoreViewportTextureAlpha(GetCapturingSceneViewport());
 	}
 }
@@ -458,8 +460,8 @@ void UAjaMediaCapture::OnFrameCaptured_RenderingThread(const FCaptureBaseData& I
 		// so don't pop from the audio buffer since we would lose these samples in the SetAudioFrameData call.
 		if (bSetVideoResult && bOutputAudio)
 		{
-			TArray<uint32> AudioSamples = AudioOutput->GetAudioSamples<uint32>();
-			OutputChannel->SetAudioFrameData(FrameBuffer, reinterpret_cast<uint8*>(AudioSamples.GetData()), AudioSamples.Num() * sizeof(uint32));
+			TArray<int32> AudioSamples = AudioOutput->GetAudioSamples<int32>();
+			OutputChannel->SetAudioFrameData(FrameBuffer, reinterpret_cast<uint8*>(AudioSamples.GetData()), AudioSamples.Num() * sizeof(int32));
 		}
 
 		if (bAjaWriteInputRawDataCmdEnable)
diff --git a/Engine/Plugins/Media/MediaIOFramework/Source/MediaIOCore/Private/MediaIOCoreAudioOutput.cpp b/Engine/Plugins/Media/MediaIOFramework/Source/MediaIOCore/Private/MediaIOCoreAudioOutput.cpp
--- a/Engine/Plugins/Media/MediaIOFramework/Source/MediaIOCore/Private/MediaIOCoreAudioOutput.cpp
+++ b/Engine/Plugins/Media/MediaIOFramework/Source/MediaIOCore/Private/MediaIOCoreAudioOutput.cpp
@@ -7,13 +7,15 @@
 #include "AudioMixerSubmix.h"
 #include "Sound/AudioSettings.h"
 
-FMediaIOAudioOutput::FMediaIOAudioOutput(Audio::FPatchOutputStrongPtr InPatchOutput, uint32 InNumInputChannels, uint32 InNumOutputChannels, FFrameRate InTargetFrameRate, uint32 InMaxSampleLatency, uint32 InOutputSampleRate)
+DEFINE_LOG_CATEGORY_STATIC(LogMediaIOAudioOutput, Log, All);
+
+FMediaIOAudioOutput::FMediaIOAudioOutput(Audio::FPatchOutputStrongPtr InPatchOutput, const FAudioOptions& InAudioOptions)
 	: PatchOutput(MoveTemp(InPatchOutput))
-    , NumInputChannels(InNumInputChannels)
-    , NumOutputChannels(InNumOutputChannels)
-    , TargetFrameRate(MoveTemp(InTargetFrameRate))
-    , MaxSampleLatency(InMaxSampleLatency)
-    , OutputSampleRate(InOutputSampleRate)
+    , NumInputChannels(InAudioOptions.InNumInputChannels)
+    , NumOutputChannels(InAudioOptions.InNumOutputChannels)
+    , TargetFrameRate(InAudioOptions.InTargetFrameRate)
+    , MaxSampleLatency(InAudioOptions.InMaxSampleLatency)
+    , OutputSampleRate(InAudioOptions.InOutputSampleRate)
 {
 }
 
@@ -38,7 +40,7 @@ Audio::FAlignedFloatBuffer FMediaIOAudioOutput::GetFloatBuffer() const
 	const int32 NumSamplesToPop = Align(NumSamplesPerFrame, 4);
 
 	Audio::FAlignedFloatBuffer FloatBuffer;
-	FloatBuffer.SetNumUninitialized(NumSamplesToPop);
+	FloatBuffer.SetNumZeroed(NumSamplesToPop);
 
 	GetAudioBuffer(NumSamplesToPop, FloatBuffer.GetData());
 
@@ -81,7 +83,11 @@ void FMediaIOAudioCapture::OnNewSubmixBuffer(const USoundSubmix* InOwningSubmix,
 	{
 		if (ensureMsgf(NumChannels == InNumChannels, TEXT("Expected %d channels from submix buffer but got %d instead."), NumChannels, InNumChannels))
 		{
-			AudioSplitter.PushAudio(InAudioData, InNumSamples);
+			int32 NumPushed = AudioSplitter.PushAudio(InAudioData, InNumSamples);
+			if (InNumSamples != NumPushed)
+			{
+				UE_LOG(LogMediaIOAudioOutput, Verbose, TEXT("Pushed samples mismatch, Incoming samples: %d, Pushed samples: %d"), InNumSamples, NumPushed);
+			}
 		}
 	}
 }
@@ -95,9 +101,16 @@ TSharedPtr<FMediaIOAudioOutput> FMediaIOAudioCapture::CreateAudioOutput(int32 In
 		checkf(NumChannels <= InNumOutputChannels, TEXT("At the moment MediaIOAudioCapture only supports up mixing."));
 		check(InNumOutputChannels != 0);
 		
-		// @todo: Don't pass strong ptr, to patch output
 		Audio::FPatchOutputStrongPtr PatchOutput = AudioSplitter.AddNewPatch(InMaxSampleLatency, Gain);
-		return MakeShared<FMediaIOAudioOutput>(MoveTemp(PatchOutput), NumChannels, InNumOutputChannels, InTargetFrameRate, InMaxSampleLatency, InOutputSampleRate);
+		FMediaIOAudioOutput::FAudioOptions Options;
+
+		Options.InNumInputChannels = NumChannels;
+		Options.InNumOutputChannels = InNumOutputChannels;
+		Options.InTargetFrameRate = InTargetFrameRate;
+		Options.InMaxSampleLatency = InMaxSampleLatency;
+		Options.InOutputSampleRate = InOutputSampleRate;
+
+		return MakeShared<FMediaIOAudioOutput>(MoveTemp(PatchOutput), Options);
 	}
 
 	return nullptr;
diff --git a/Engine/Plugins/Media/MediaIOFramework/Source/MediaIOCore/Private/MediaIOCoreFileWriter.cpp b/Engine/Plugins/Media/MediaIOFramework/Source/MediaIOCore/Private/MediaIOCoreFileWriter.cpp
--- a/Engine/Plugins/Media/MediaIOFramework/Source/MediaIOCore/Private/MediaIOCoreFileWriter.cpp
+++ b/Engine/Plugins/Media/MediaIOFramework/Source/MediaIOCore/Private/MediaIOCoreFileWriter.cpp
@@ -11,10 +11,10 @@
 
 namespace MediaIOCoreFileWriter
 {
-	void WriteRawFile(const FString& InFilename, uint8* InBuffer, uint32 InSize)
+	void WriteRawFile(const FString& InFilename, uint8* InBuffer, uint32 InSize, bool bAppend)
 	{
 #if ALLOW_DEBUG_FILES
-		if (InFilename != "")
+		if (!InFilename.IsEmpty())
 		{
 			FString OutputDirectory = FPaths::Combine(*FPaths::ProjectDir(), TEXT("Saved"), TEXT("Media"));
 			FPaths::NormalizeDirectoryName(OutputDirectory);
@@ -31,13 +31,20 @@ namespace MediaIOCoreFileWriter
 
 			// Append current date and time
 			FDateTime CurrentDateAndTime = FDateTime::Now();
-			OutputFilename += FString("_") + CurrentDateAndTime.ToString() + FString(".raw");
+			if (bAppend)
+			{
+				OutputFilename += FString(".raw");
+			}
+			else
+			{
+				OutputFilename += FString("_") + CurrentDateAndTime.ToString() + FString(".raw");
+			}
 
 			IPlatformFile& PlatformFile = FPlatformFileManager::Get().GetPlatformFile();
-			if (!PlatformFile.FileExists(*OutputFilename))
+			if (bAppend || !PlatformFile.FileExists(*OutputFilename))
 			{
 				TUniquePtr<IFileHandle> FileHandle;
-				FileHandle.Reset(PlatformFile.OpenWrite(*OutputFilename));
+				FileHandle.Reset(PlatformFile.OpenWrite(*OutputFilename, bAppend));
 				if (FileHandle)
 				{
 					FileHandle->Write(InBuffer, InSize);
diff --git a/Engine/Plugins/Media/MediaIOFramework/Source/MediaIOCore/Public/MediaIOCoreAudioOutput.h b/Engine/Plugins/Media/MediaIOFramework/Source/MediaIOCore/Public/MediaIOCoreAudioOutput.h
--- a/Engine/Plugins/Media/MediaIOFramework/Source/MediaIOCore/Public/MediaIOCoreAudioOutput.h
+++ b/Engine/Plugins/Media/MediaIOFramework/Source/MediaIOCore/Public/MediaIOCoreAudioOutput.h
@@ -46,9 +46,18 @@ namespace UE::MediaIoCoreModule::Private
 }
 
 class MEDIAIOCORE_API FMediaIOAudioOutput
-{
+{ 
 public:
-	FMediaIOAudioOutput(Audio::FPatchOutputStrongPtr InPatchOutput, uint32 InNumInputChannels, uint32 InNumOutputChannels, FFrameRate InTargetFrameRate, uint32 InMaxSampleLatency, uint32 InOutputSampleRate);
+	struct FAudioOptions
+	{
+		uint32 InNumInputChannels = 0;
+		uint32 InNumOutputChannels = 0;
+		FFrameRate InTargetFrameRate;
+		uint32 InMaxSampleLatency = 0;
+		uint32 InOutputSampleRate = 0;
+	};
+
+	FMediaIOAudioOutput(Audio::FPatchOutputStrongPtr InPatchOutput, const FAudioOptions& InAudioOptions);
 
 	/**
 	 * Get the audio sample that were accumulated.
@@ -69,13 +78,18 @@ private:
 private:
 	/** The buffer accumulating audio samples. */
 	Audio::FPatchOutputStrongPtr PatchOutput;
+
 	/** Number of audio channels on the engine side. */
 	int32 NumInputChannels;
+
 	/** Number of audio channels to output. */
 	int32 NumOutputChannels;
+
 	FFrameRate TargetFrameRate;
+
 	/** Maximum number of samples to accumulate before they are discarded. */
 	uint32 MaxSampleLatency;
+
 	uint32 OutputSampleRate;
 };
 
@@ -97,10 +111,13 @@ public:
 private:
 	/** Sample rate on the engine side. */ 
 	uint32 SampleRate = 0;
+
 	// Used to make sure we only accumulate audio from the master submix. 
 	FName MasterSubmixName;
+
 	/** Number of channels on the engine side. */
 	int32 NumChannels = 0;
+
 	/** Utility that allows pushing audio samples to multiple outputs. */
 	Audio::FPatchSplitter AudioSplitter;
 };
diff --git a/Engine/Plugins/Media/MediaIOFramework/Source/MediaIOCore/Public/MediaIOCoreFileWriter.h b/Engine/Plugins/Media/MediaIOFramework/Source/MediaIOCore/Public/MediaIOCoreFileWriter.h
--- a/Engine/Plugins/Media/MediaIOFramework/Source/MediaIOCore/Public/MediaIOCoreFileWriter.h
+++ b/Engine/Plugins/Media/MediaIOFramework/Source/MediaIOCore/Public/MediaIOCoreFileWriter.h
@@ -6,5 +6,5 @@
 
 namespace MediaIOCoreFileWriter
 {
-	MEDIAIOCORE_API void WriteRawFile(const FString& InFilename, uint8* InBuffer, uint32 InSize);
+	MEDIAIOCORE_API void WriteRawFile(const FString& InFilename, uint8* InBuffer, uint32 InSize, bool bAppend = true);
 };
\ No newline at end of file
