Add support to predetermine the number of frame to capture

diff --git a/Engine/Plugins/Media/MediaIOFramework/Source/MediaIOCore/Private/MediaCapture.cpp b/Engine/Plugins/Media/MediaIOFramework/Source/MediaIOCore/Private/MediaCapture.cpp
--- a/Engine/Plugins/Media/MediaIOFramework/Source/MediaIOCore/Private/MediaCapture.cpp
+++ b/Engine/Plugins/Media/MediaIOFramework/Source/MediaIOCore/Private/MediaCapture.cpp
@@ -148,6 +148,8 @@ FMediaCaptureOptions::FMediaCaptureOptions()
 	, bSkipFrameWhenRunningExpensiveTasks(true)
 	, bConvertToDesiredPixelFormat(true)
 	, bForceAlphaToOneOnConversion(false)
+	, bAutostopOnCapture(false)
+	, NumberOfFramesToCapture(-1)
 {
 
 }
@@ -219,13 +221,14 @@ bool UMediaCapture::CaptureSceneViewport(TSharedPtr<FSceneViewport>& InSceneView
 
 	check(IsInGameThread());
 
+	DesiredCaptureOptions = InCaptureOptions;
+
 	if (!ValidateMediaOutput())
 	{
 		MediaCaptureDetails::ShowSlateNotification();
 		return false;
 	}
 
-	DesiredCaptureOptions = InCaptureOptions;
 	CacheMediaOutput(EMediaCaptureSourceType::SCENE_VIEWPORT);
 
 	if (bUseRequestedTargetSize)
@@ -287,13 +290,14 @@ bool UMediaCapture::CaptureTextureRenderTarget2D(UTextureRenderTarget2D* InRende
 
 	check(IsInGameThread());
 
+	DesiredCaptureOptions = CaptureOptions; 
+	
 	if (!ValidateMediaOutput())
 	{
 		MediaCaptureDetails::ShowSlateNotification();
 		return false;
 	}
 
-	DesiredCaptureOptions = CaptureOptions;
 	CacheMediaOutput(EMediaCaptureSourceType::RENDER_TARGET);
 
 	if (bUseRequestedTargetSize)
@@ -503,12 +507,16 @@ void UMediaCapture::StopCapture(bool bAllowPendingFrameToBeProcess)
 		{
 			SetState(EMediaCaptureState::StopRequested);
 
+			//Do not flush when auto stopping to avoid hitches.
+			if(DesiredCaptureOptions.bAutostopOnCapture != true)
+			{
 			while (WaitingForResolveCommandExecutionCounter.Load() > 0)
 			{
 				FlushRenderingCommands();
 			}
 		}
 	}
+	}
 	else
 	{
 		if (GetState() != EMediaCaptureState::Stopped)
@@ -622,13 +630,13 @@ void UMediaCapture::InitializeResolveTarget(int32 InNumberOfBuffers)
 	check(CaptureFrames.Num() == 0);
 	CaptureFrames.AddDefaulted(InNumberOfBuffers);
 
-	UMediaCapture* This = this;
-	ENQUEUE_RENDER_COMMAND(MediaOutputCaptureFrameCreateTexture)(
-		[This](FRHICommandListImmediate& RHICmdList)
-		{
-			FRHIResourceCreateInfo CreateInfo(TEXT("UMediaCapture"));
-			for (int32 Index = 0; Index < This->NumberOfCaptureFrame; ++Index)
+		UMediaCapture* This = this;
+		ENQUEUE_RENDER_COMMAND(MediaOutputCaptureFrameCreateTexture)(
+			[This](FRHICommandListImmediate& RHICmdList)
 			{
+				FRHIResourceCreateInfo CreateInfo(TEXT("UMediaCapture"));
+				for (int32 Index = 0; Index < This->NumberOfCaptureFrame; ++Index)
+				{
 				FPooledRenderTargetDesc OutputDesc = FPooledRenderTargetDesc::Create2DDesc(
 					This->DesiredOutputSize,
 					This->DesiredOutputPixelFormat,
@@ -654,8 +662,8 @@ void UMediaCapture::InitializeResolveTarget(int32 InNumberOfBuffers)
 				}
 
 			}
-			This->bResolvedTargetInitialized = true;
-		});
+				This->bResolvedTargetInitialized = true;
+			});
 }
 
 bool UMediaCapture::ValidateMediaOutput() const
@@ -673,6 +681,12 @@ bool UMediaCapture::ValidateMediaOutput() const
 		return false;
 	}
 
+	if(DesiredCaptureOptions.bAutostopOnCapture && DesiredCaptureOptions.NumberOfFramesToCapture < 1)
+	{
+		UE_LOG(LogMediaIOCore, Error, TEXT("Can not start the capture. Please set the Number Of Frames To Capture when using Autostop On Capture in the Media Capture Options"));
+		return false;
+	}
+
 	return true;
 }
 
@@ -757,6 +771,12 @@ void UMediaCapture::OnEndFrame_GameThread()
 			{
 				InMediaCapture->Capture_RenderThread(RHICmdList, InMediaCapture, CapturingFrame, ReadyFrame, InCapturingSceneViewport, InTextureRenderTargetResource, InDesiredSize, InOnStateChanged);
 			});
+
+			//If auto-stopping, count the number of frame captures requested and stop when reaching 0.
+			if(DesiredCaptureOptions.bAutostopOnCapture && GetState() == EMediaCaptureState::Capturing && --DesiredCaptureOptions.NumberOfFramesToCapture <= 0)
+			{
+				StopCapture(true);
+			}
 		}
 	}
 }
diff --git a/Engine/Plugins/Media/MediaIOFramework/Source/MediaIOCore/Public/MediaCapture.h b/Engine/Plugins/Media/MediaIOFramework/Source/MediaIOCore/Public/MediaCapture.h
--- a/Engine/Plugins/Media/MediaIOFramework/Source/MediaIOCore/Public/MediaCapture.h
+++ b/Engine/Plugins/Media/MediaIOFramework/Source/MediaIOCore/Public/MediaCapture.h
@@ -117,6 +117,16 @@ public:
 	 */
 	UPROPERTY(BlueprintReadWrite, EditAnywhere, Category="MediaCapture")
 	bool bForceAlphaToOneOnConversion;
+
+	/** Automatically stop capturing after a predetermined number of images. */
+	UPROPERTY(BlueprintReadWrite, EditAnywhere, Category = "MediaCapture")
+	bool bAutostopOnCapture;
+
+	/** The number of images to capture*/
+	UPROPERTY(BlueprintReadWrite, EditAnywhere, Category = "MediaCapture")
+	int32 NumberOfFramesToCapture;
+
+
 };
 
 
