Fix no sound output in PIE, as well as audio glitches when no sound is being played in the engine when using aja

With my testing I didn't find a perf impact by setting the memory to zero when not outputting.

diff --git a/Engine/Plugins/Media/MediaIOFramework/Source/MediaIOCore/Private/MediaIOCoreAudioOutput.cpp b/Engine/Plugins/Media/MediaIOFramework/Source/MediaIOCore/Private/MediaIOCoreAudioOutput.cpp
--- a/Engine/Plugins/Media/MediaIOFramework/Source/MediaIOCore/Private/MediaIOCoreAudioOutput.cpp
+++ b/Engine/Plugins/Media/MediaIOFramework/Source/MediaIOCore/Private/MediaIOCoreAudioOutput.cpp
@@ -7,6 +7,11 @@
 #include "AudioMixerSubmix.h"
 #include "Sound/AudioSettings.h"
 
+
+#if WITH_EDITOR
+#include "Editor.h"
+#endif
+
 DEFINE_LOG_CATEGORY_STATIC(LogMediaIOAudioOutput, Log, All);
 
 FMediaIOAudioOutput::FMediaIOAudioOutput(Audio::FPatchOutputStrongPtr InPatchOutput, const FAudioOptions& InAudioOptions)
@@ -42,38 +47,31 @@ Audio::FAlignedFloatBuffer FMediaIOAudioOutput::GetFloatBuffer() const
 	Audio::FAlignedFloatBuffer FloatBuffer;
 	FloatBuffer.SetNumZeroed(NumSamplesToPop);
 
-	GetAudioBuffer(NumSamplesToPop, FloatBuffer.GetData());
+	const int32 NumPopped = GetAudioBuffer(NumSamplesToPop, FloatBuffer.GetData());
 
 	// Trim back the buffer after we get the buffer since it might be bigger because of alignment
 	constexpr bool bAllowShrinking = false;
-	FloatBuffer.SetNum(NumSamplesPerFrame, bAllowShrinking);
+	FloatBuffer.SetNum(NumPopped, bAllowShrinking);
 	return FloatBuffer;
 }
 
 FMediaIOAudioCapture::FMediaIOAudioCapture()
 {
-	if (FAudioDevice* AudioDevice = GEngine->GetMainAudioDeviceRaw())
-	{
-		if (AudioDevice->IsAudioMixerEnabled())
-		{
-			Audio::FMixerDevice* MixerDevice = static_cast<Audio::FMixerDevice*>(AudioDevice);
-			NumChannels = MixerDevice->GetDeviceOutputChannels();
-			SampleRate = MixerDevice->GetSampleRate();
-			MasterSubmixName = *GetDefault<UAudioSettings>()->MasterSubmix.GetAssetName();
-			AudioDevice->RegisterSubmixBufferListener(this);
-		}
-	}
+	RegisterMainAudioDevice();
+#if WITH_EDITOR
+	FEditorDelegates::PostPIEStarted.AddRaw(this, &FMediaIOAudioCapture::OnPIEStarted);
+	FEditorDelegates::PrePIEEnded.AddRaw(this, &FMediaIOAudioCapture::OnPIEEnded);
+#endif
 }
 
 FMediaIOAudioCapture::~FMediaIOAudioCapture()
 {
-	if (FAudioDeviceManager* DeviceManager = FAudioDeviceManager::Get())
-	{
-		if (FAudioDevice* AudioDevice = GEngine->GetMainAudioDeviceRaw())
-		{
-			AudioDevice->UnregisterSubmixBufferListener(this);
-		}
-	}
+#if WITH_EDITOR
+	FEditorDelegates::EndPIE.RemoveAll(this);
+	FEditorDelegates::PostPIEStarted.RemoveAll(this);
+#endif
+
+	UnregisterMainAudioDevice();
 }
 
 void FMediaIOAudioCapture::OnNewSubmixBuffer(const USoundSubmix* InOwningSubmix, float* InAudioData, int32 InNumSamples, int32 InNumChannels, const int32 InSampleRate, double InAudioClock)
@@ -115,3 +113,62 @@ TSharedPtr<FMediaIOAudioOutput> FMediaIOAudioCapture::CreateAudioOutput(int32 In
 
 	return nullptr;
 }
+
+#if WITH_EDITOR
+void FMediaIOAudioCapture::OnPIEStarted(const bool)
+{
+	UnregisterMainAudioDevice();
+	
+	if (GEditor)
+	{
+		if (FWorldContext* PIEWorldContext = GEditor->GetPIEWorldContext())
+		{
+			Audio::FMixerDevice* MixerDevice = static_cast<Audio::FMixerDevice*>(PIEWorldContext->World()->GetAudioDeviceRaw());
+			RegisterBufferListener(MixerDevice);
+		}
+	}
+}
+
+void FMediaIOAudioCapture::OnPIEEnded(const bool)
+{
+	if (GEditor)
+	{
+		if (FWorldContext* PIEWorldContext = GEditor->GetPIEWorldContext())
+		{
+			UnregisterBufferListener(PIEWorldContext->World()->GetAudioDeviceRaw());
+		}
+	}
+	
+	RegisterMainAudioDevice();
+}
+#endif
+
+void FMediaIOAudioCapture::RegisterMainAudioDevice()
+{
+	RegisterBufferListener(GEngine->GetMainAudioDeviceRaw());
+}
+
+void FMediaIOAudioCapture::UnregisterMainAudioDevice()
+{
+	UnregisterBufferListener(GEngine->GetMainAudioDeviceRaw());
+}
+
+void FMediaIOAudioCapture::RegisterBufferListener(FAudioDevice* AudioDevice)
+{
+	if (AudioDevice && AudioDevice->IsAudioMixerEnabled())
+	{
+		Audio::FMixerDevice* MixerDevice = static_cast<Audio::FMixerDevice*>(AudioDevice);
+		NumChannels = MixerDevice->GetDeviceOutputChannels();
+		SampleRate = MixerDevice->GetSampleRate();
+		MasterSubmixName = *GetDefault<UAudioSettings>()->MasterSubmix.GetAssetName();
+		AudioDevice->RegisterSubmixBufferListener(this);
+	}
+}
+
+void FMediaIOAudioCapture::UnregisterBufferListener(FAudioDevice* AudioDevice)
+{
+	if (AudioDevice)
+	{
+		AudioDevice->UnregisterSubmixBufferListener(this);
+	}
+}
diff --git a/Engine/Plugins/Media/MediaIOFramework/Source/MediaIOCore/Public/MediaIOCoreAudioOutput.h b/Engine/Plugins/Media/MediaIOFramework/Source/MediaIOCore/Public/MediaIOCoreAudioOutput.h
--- a/Engine/Plugins/Media/MediaIOFramework/Source/MediaIOCore/Public/MediaIOCoreAudioOutput.h
+++ b/Engine/Plugins/Media/MediaIOFramework/Source/MediaIOCore/Public/MediaIOCoreAudioOutput.h
@@ -25,17 +25,18 @@ namespace UE::MediaIoCoreModule::Private
 
 		OutputType* ConvertedBufferPtr = ConvertedBuffer.GetData();
 
-		for (int32 Index = 0; Index < InBuffer.Num(); Index += 2)
+		for (int32 Index = 0; Index < InBuffer.Num(); Index += NumInputChannels)
 		{
 			constexpr double Scale = TNumericLimits<OutputType>::Max();
-		
+
+			// @Note: While this conversion is a common approach, it can introduce distortion. 
+			// See: https://www.cs.cmu.edu/~rbd/papers/cmj-float-to-int.html
 			// Convert values for each channel
 			for (int32 ChannelIndex = 0; ChannelIndex < NumInputChannels; ++ChannelIndex)
 			{
-				// @Note: While this conversion is a common approach, it can introduce distortion. 
-				// See: https://www.cs.cmu.edu/~rbd/papers/cmj-float-to-int.html
 				const float FloatValue = InBuffer[Index + ChannelIndex];
 				OutputType ConvertedValue = static_cast<OutputType>((FloatValue * Scale) + 0.5);
+
 				*(ConvertedBufferPtr + ChannelIndex) = ConvertedValue;
 			}
 		
@@ -108,6 +109,16 @@ public:
 	/** Create an audio output that will receive audio samples. */
 	TSharedPtr<FMediaIOAudioOutput> CreateAudioOutput(int32 InNumOutputChannels, FFrameRate InTargetFrameRate, uint32 InMaxSampleLatency, uint32 InOutputSampleRate);
 
+private:
+#if WITH_EDITOR
+	void OnPIEStarted(const bool);
+	void OnPIEEnded(const bool);
+#endif
+
+	void RegisterMainAudioDevice();
+	void UnregisterMainAudioDevice();
+	void RegisterBufferListener(FAudioDevice* AudioDevice);
+	void UnregisterBufferListener(FAudioDevice* AudioDevice);
 private:
 	/** Sample rate on the engine side. */ 
 	uint32 SampleRate = 0;
