Convert gpu texture transfer warning to log

diff --git a/Engine/Plugins/Media/MediaIOFramework/Source/GPUTextureTransfer/Private/GPUTextureTransferModule.cpp b/Engine/Plugins/Media/MediaIOFramework/Source/GPUTextureTransfer/Private/GPUTextureTransferModule.cpp
--- a/Engine/Plugins/Media/MediaIOFramework/Source/GPUTextureTransfer/Private/GPUTextureTransferModule.cpp
+++ b/Engine/Plugins/Media/MediaIOFramework/Source/GPUTextureTransfer/Private/GPUTextureTransferModule.cpp
@@ -99,7 +99,7 @@ bool FGPUTextureTransferModule::LoadGPUDirectBinary()
 	TextureTransferHandle = FPlatformProcess::GetDllHandle(*DVPDll);
 	if (TextureTransferHandle == nullptr)
 	{
-		UE_LOG(LogGPUTextureTransfer, Warning, TEXT("Failed to load required library %s. GPU Texture transfer will not be functional."), *DVPDll);
+		UE_LOG(LogGPUTextureTransfer, Display, TEXT("Failed to load required library %s. GPU Texture transfer will not be functional."), *DVPDll);
 	}
 
 	FPlatformProcess::PopDllDirectory(*GPUDirectPath);
