Fix static analysis warning

diff --git a/Engine/Plugins/Media/MediaIOFramework/Source/GPUTextureTransfer/Private/TextureTransferBase.h b/Engine/Plugins/Media/MediaIOFramework/Source/GPUTextureTransfer/Private/TextureTransferBase.h
--- a/Engine/Plugins/Media/MediaIOFramework/Source/GPUTextureTransfer/Private/TextureTransferBase.h
+++ b/Engine/Plugins/Media/MediaIOFramework/Source/GPUTextureTransfer/Private/TextureTransferBase.h
@@ -67,7 +67,7 @@ namespace UE::GPUTextureTransfer::Private
 			DVPBufferHandle DVPHandle = 0;
 			union 
 			{
-				void* Handle;
+				void* Handle = nullptr;
 				int Fd;
 			} External; // D3D12 and Vulkan only
 		};
