VulkanRHI:  Add stubbed VulkanTextureTransfer so that GPUDirect initialization doesn't assert on startup of Windows Vulkan.

diff --git a/Engine/Plugins/Media/MediaIOFramework/Source/GPUTextureTransfer/GPUTextureTransfer.Build.cs b/Engine/Plugins/Media/MediaIOFramework/Source/GPUTextureTransfer/GPUTextureTransfer.Build.cs
--- a/Engine/Plugins/Media/MediaIOFramework/Source/GPUTextureTransfer/GPUTextureTransfer.Build.cs
+++ b/Engine/Plugins/Media/MediaIOFramework/Source/GPUTextureTransfer/GPUTextureTransfer.Build.cs
@@ -9,8 +9,9 @@ namespace UnrealBuildTool.Rules
 			if (Target.Platform == UnrealTargetPlatform.Win64)
 			{
 				AddEngineThirdPartyPrivateStaticDependencies(Target, "GPUDirect");
+				AddEngineThirdPartyPrivateStaticDependencies(Target, "Vulkan");
 			}
-			
+
 			PublicDefinitions.Add("PERF_LOGGING=0");
 
 			PublicDependencyModuleNames.AddRange(
@@ -34,6 +35,7 @@ namespace UnrealBuildTool.Rules
 					"GPUDirect",
 					"RenderCore",
 					"RHI",
+					"VulkanRHI",
 				});
 
 			if (Target.bBuildEditor == true)
diff --git a/Engine/Plugins/Media/MediaIOFramework/Source/GPUTextureTransfer/Private/GPUTextureTransferModule.cpp b/Engine/Plugins/Media/MediaIOFramework/Source/GPUTextureTransfer/Private/GPUTextureTransferModule.cpp
--- a/Engine/Plugins/Media/MediaIOFramework/Source/GPUTextureTransfer/Private/GPUTextureTransferModule.cpp
+++ b/Engine/Plugins/Media/MediaIOFramework/Source/GPUTextureTransfer/Private/GPUTextureTransferModule.cpp
@@ -4,8 +4,10 @@
 
 #include "D3D11TextureTransfer.h"
 #include "D3D12TextureTransfer.h"
+#include "VulkanTextureTransfer.h"
 #include "RHI.h"
 #include "DynamicRHI.h"
+#include "IVulkanDynamicRHI.h"
 #include "GenericPlatform/GenericPlatformDriver.h"
 #include "HAL/PlatformMisc.h"
 #include "Misc/CoreDelegates.h"
@@ -23,6 +25,7 @@ namespace
 		{
 		case ERHIInterfaceType::D3D11: return UE::GPUTextureTransfer::ERHI::D3D11;
 		case ERHIInterfaceType::D3D12: return UE::GPUTextureTransfer::ERHI::D3D12;
+		case ERHIInterfaceType::Vulkan: return UE::GPUTextureTransfer::ERHI::Vulkan;
 		default: return UE::GPUTextureTransfer::ERHI::Invalid;
 		}
 	};
@@ -130,6 +133,9 @@ void FGPUTextureTransferModule::InitializeTextureTransfer()
 		case UE::GPUTextureTransfer::ERHI::D3D12:
 			TextureTransfer = MakeShared<UE::GPUTextureTransfer::Private::FD3D12TextureTransfer>();
 			break;
+		case UE::GPUTextureTransfer::ERHI::Vulkan:
+			TextureTransfer = MakeShared<UE::GPUTextureTransfer::Private::FVulkanTextureTransfer>();
+			break;
 		default:
 			ensureAlways(false);
 			break;
@@ -139,15 +145,12 @@ void FGPUTextureTransferModule::InitializeTextureTransfer()
 		InitializeArgs.RHI = RHI;
 		InitializeArgs.RHIDevice = GDynamicRHI->RHIGetNativeDevice();
 		InitializeArgs.RHICommandQueue = GDynamicRHI->RHIGetNativeGraphicsQueue();
-		/* Re-enable when adding vulkan support
-		if (RHI == AJA::ERHI::Vulkan)
+		if (RHI == UE::GPUTextureTransfer::ERHI::Vulkan)
 		{
-			FVulkanDynamicRHI* vkDynamicRHI = GetDynamicRHI<FVulkanDynamicRHI>();
-			InitializeArgs.VulkanInstance = vkDynamicRHI->GetInstance();
-
-			FMemory::Memcpy(InitializeArgs.RHIDeviceUUID, vkDynamicRHI->GetDevice()->GetDeviceIdProperties().deviceUUID, 16);
+			IVulkanDynamicRHI* DynRHI = GetIVulkanDynamicRHI();
+			InitializeArgs.VulkanInstance = DynRHI->RHIGetVkInstance();
+			FMemory::Memcpy(InitializeArgs.RHIDeviceUUID, DynRHI->RHIGetVulkanDeviceUUID(), 16);
 		}
-		*/
 
 		const uint8 RHIIndex = static_cast<uint8>(RHI);
 		if (TextureTransfer->Initialize(InitializeArgs))
diff --git a/Engine/Plugins/Media/MediaIOFramework/Source/GPUTextureTransfer/Private/VulkanTextureTransfer.cpp b/Engine/Plugins/Media/MediaIOFramework/Source/GPUTextureTransfer/Private/VulkanTextureTransfer.cpp
new file mode 100644
--- /dev/null
+++ b/Engine/Plugins/Media/MediaIOFramework/Source/GPUTextureTransfer/Private/VulkanTextureTransfer.cpp
@@ -0,0 +1,110 @@
+// Copyright Epic Games, Inc. All Rights Reserved.
+
+#include "VulkanTextureTransfer.h"
+
+#if PLATFORM_WINDOWS
+#include "Windows/AllowWindowsPlatformTypes.h"
+#include "Windows/PreWindowsApi.h"
+#endif
+
+#include "DVPAPI.h"
+#include "dvpapi_vulkan.h"
+
+#if PLATFORM_WINDOWS
+#include "Windows/PostWindowsApi.h"
+#include "Windows/HideWindowsPlatformTypes.h"
+#endif
+
+
+namespace UE::GPUTextureTransfer::Private
+{
+	DVPStatus FVulkanTextureTransfer::Init_Impl(const FInitializeDMAArgs& InArgs)
+	{
+		VulkanDevice = (VkDevice)InArgs.RHIDevice;
+		VulkanQueue = (VkQueue)InArgs.RHICommandQueue;
+		return dvpInitVkDevice(VulkanDevice, 0, (void*)InArgs.RHIDeviceUUID);
+	}
+
+	DVPStatus FVulkanTextureTransfer::GetConstants_Impl(uint32* OutBufferAddrAlignment, uint32* OutBufferGPUStrideAlignment, uint32* OutSemaphoreAddrAlignment, uint32* OutSemaphoreAllocSize, uint32* OutSemaphorePayloadOffset, uint32* OutSemaphorePayloadSize) const
+	{
+		if (!VulkanDevice)
+		{
+			return DVP_STATUS_ERROR;
+		}
+
+		return dvpGetRequiredConstantsVkDevice(OutBufferAddrAlignment, OutBufferGPUStrideAlignment, OutSemaphoreAddrAlignment, OutSemaphoreAllocSize,
+			OutSemaphorePayloadOffset, OutSemaphorePayloadSize, VulkanDevice);
+	}
+
+	DVPStatus FVulkanTextureTransfer::CloseDevice_Impl() const
+	{
+		if (!VulkanDevice)
+		{
+			return DVP_STATUS_ERROR;
+		}
+
+		return dvpCloseVkDevice(VulkanDevice);
+	}
+
+	DVPStatus FVulkanTextureTransfer::BindBuffer_Impl(DVPBufferHandle InBufferHandle) const
+	{
+		if (!VulkanDevice)
+		{
+			return DVP_STATUS_ERROR;
+		}
+
+		return dvpBindToVkDevice(InBufferHandle, VulkanDevice);
+	}
+
+	DVPStatus FVulkanTextureTransfer::CreateGPUResource_Impl(void* InTexture, FTextureTransferBase::FTextureInfo* OutTextureInfo) const
+	{
+		if (!OutTextureInfo || !VulkanDevice)
+		{
+			return DVP_STATUS_ERROR;
+		}
+
+		// todo for Vulkan
+		checkNoEntry();
+
+		DVPGpuExternalResourceDesc Desc;
+		//Desc.width = (uint32) ResourceDesc.Width;
+		//Desc.height = (uint32)ResourceDesc.Height;
+		//Desc.size = Desc.width * Desc.height * 4;
+		Desc.format = DVP_BGRA;
+		Desc.type = DVP_UNSIGNED_BYTE;
+		Desc.handleType = DVP_OPAQUE_WIN32;
+		Desc.external.handle = OutTextureInfo->External.Handle;
+
+		return dvpCreateGPUExternalResourceVkDevice(VulkanDevice, &Desc, &OutTextureInfo->DVPHandle);
+	}
+
+	DVPStatus FVulkanTextureTransfer::UnbindBuffer_Impl(DVPBufferHandle InBufferHandle) const
+	{
+		if (!VulkanDevice)
+		{
+			return DVP_STATUS_ERROR;
+		}
+
+		return dvpUnbindFromVkDevice(InBufferHandle, VulkanDevice);
+	}
+
+	DVPStatus FVulkanTextureTransfer::MapBufferWaitAPI_Impl(DVPBufferHandle InHandle) const
+	{
+		if (!VulkanQueue)
+		{
+			return DVP_STATUS_ERROR;
+		}
+
+		return (dvpMapBufferWaitVk(InHandle, VulkanQueue));
+	}
+
+	DVPStatus FVulkanTextureTransfer::MapBufferEndAPI_Impl(DVPBufferHandle InHandle) const
+	{
+		if (!VulkanQueue)
+		{
+			return DVP_STATUS_ERROR;
+		}
+
+		return (dvpMapBufferEndVk(InHandle, VulkanQueue));
+	}
+}
diff --git a/Engine/Plugins/Media/MediaIOFramework/Source/GPUTextureTransfer/Private/VulkanTextureTransfer.h b/Engine/Plugins/Media/MediaIOFramework/Source/GPUTextureTransfer/Private/VulkanTextureTransfer.h
new file mode 100644
--- /dev/null
+++ b/Engine/Plugins/Media/MediaIOFramework/Source/GPUTextureTransfer/Private/VulkanTextureTransfer.h
@@ -0,0 +1,28 @@
+// Copyright Epic Games, Inc. All Rights Reserved.
+
+#pragma once
+
+#include "TextureTransferBase.h"
+#include "IVulkanDynamicRHI.h"
+
+namespace UE::GPUTextureTransfer::Private
+{
+class FVulkanTextureTransfer : public FTextureTransferBase
+{
+protected:
+	virtual DVPStatus Init_Impl(const FInitializeDMAArgs& InArgs) override;
+	virtual DVPStatus GetConstants_Impl(uint32* OutBufferAddrAlignment, uint32* OutBufferGPUStrideAlignment, uint32* OutSemaphoreAddrAlignment, uint32* OutSemaphoreAllocSize, uint32* OutSemaphorePayloadOffset, uint32* OutSemaphorePayloadSize) const override;
+	virtual DVPStatus BindBuffer_Impl(DVPBufferHandle InBufferHandle) const override;
+	virtual DVPStatus CreateGPUResource_Impl(void* InTexture, FTextureInfo* OutTextureInfo) const override;
+	virtual DVPStatus UnbindBuffer_Impl(DVPBufferHandle InBufferHandle) const override;
+	virtual DVPStatus CloseDevice_Impl() const override;
+	virtual DVPStatus MapBufferWaitAPI_Impl(DVPBufferHandle InHandle) const;
+	virtual DVPStatus MapBufferEndAPI_Impl(DVPBufferHandle InHandle) const;
+
+
+private:
+	VkInstance VulkanInstance = nullptr;
+	VkDevice VulkanDevice = nullptr;
+	VkQueue VulkanQueue = nullptr;
+};
+}
\ No newline at end of file
diff --git a/Engine/Plugins/Media/MediaIOFramework/Source/GPUTextureTransfer/Public/GPUTextureTransfer.h b/Engine/Plugins/Media/MediaIOFramework/Source/GPUTextureTransfer/Public/GPUTextureTransfer.h
--- a/Engine/Plugins/Media/MediaIOFramework/Source/GPUTextureTransfer/Public/GPUTextureTransfer.h
+++ b/Engine/Plugins/Media/MediaIOFramework/Source/GPUTextureTransfer/Public/GPUTextureTransfer.h
@@ -11,6 +11,7 @@ namespace UE::GPUTextureTransfer
 		Invalid,
 		D3D11,
 		D3D12,
+		Vulkan,
 		RHI_MAX
 	};
 
