Fix CIS warning in gpu texture transfer

diff --git a/Engine/Plugins/Media/MediaIOFramework/Source/GPUTextureTransfer/Private/TextureTransferBase.cpp b/Engine/Plugins/Media/MediaIOFramework/Source/GPUTextureTransfer/Private/TextureTransferBase.cpp
--- a/Engine/Plugins/Media/MediaIOFramework/Source/GPUTextureTransfer/Private/TextureTransferBase.cpp
+++ b/Engine/Plugins/Media/MediaIOFramework/Source/GPUTextureTransfer/Private/TextureTransferBase.cpp
@@ -278,15 +278,15 @@ namespace UE::GPUTextureTransfer::Private
 
 	void FTextureTransferBase::ClearRegisteredTextures()
 	{
-		for (const TPair<void*, FTextureInfo>& Pair : RegisteredTextures)
+		for (auto It = RegisteredTextures.CreateConstIterator(); It; ++It)
 		{
-			if (Pair.Value.DVPHandle)
+			if (It.Value().DVPHandle)
 			{
-				dvpFreeBuffer(Pair.Value.DVPHandle);
-				if (Pair.Value.External.Handle)
+				dvpFreeBuffer(It.Value().DVPHandle);
+				if (It.Value().External.Handle)
 				{
 #if PLATFORM_WINDOWS
-					CloseHandle(Pair.Value.External.Handle);
+					CloseHandle(It.Value().External.Handle);
 #endif
 				}
 			}
