Change another log error to display

diff --git a/Engine/Plugins/Media/MediaIOFramework/Source/GPUTextureTransfer/Private/TextureTransferBase.cpp b/Engine/Plugins/Media/MediaIOFramework/Source/GPUTextureTransfer/Private/TextureTransferBase.cpp
--- a/Engine/Plugins/Media/MediaIOFramework/Source/GPUTextureTransfer/Private/TextureTransferBase.cpp
+++ b/Engine/Plugins/Media/MediaIOFramework/Source/GPUTextureTransfer/Private/TextureTransferBase.cpp
@@ -72,7 +72,7 @@ namespace UE::GPUTextureTransfer::Private
 
 		if (Init_Impl(Args) != DVP_STATUS_OK)
 		{
-			UE_LOG(LogGPUTextureTransfer, Error, TEXT("GPU Direct failed to initialize."));
+			UE_LOG(LogGPUTextureTransfer, Display, TEXT("GPU Direct failed to initialize."));
 			return false;
 		}
 
