Fix linux build

diff --git a/Engine/Plugins/Media/MediaIOFramework/Source/GPUTextureTransfer/GPUTextureTransfer.Build.cs b/Engine/Plugins/Media/MediaIOFramework/Source/GPUTextureTransfer/GPUTextureTransfer.Build.cs
--- a/Engine/Plugins/Media/MediaIOFramework/Source/GPUTextureTransfer/GPUTextureTransfer.Build.cs
+++ b/Engine/Plugins/Media/MediaIOFramework/Source/GPUTextureTransfer/GPUTextureTransfer.Build.cs
@@ -13,7 +13,7 @@ namespace UnrealBuildTool.Rules
 				PrivateDependencyModuleNames.Add("VulkanRHI");
 				PublicDefinitions.Add("DVP_SUPPORTED_PLATFORM=1");
 			}
-			else if (Target.Platform == UnrealTargetPlatform.Linux)
+			else if (Target.Platform == UnrealTargetPlatform.Linux || Target.Platform == UnrealTargetPlatform.LinuxArm64)
 			{
 				AddEngineThirdPartyPrivateStaticDependencies(Target, "Vulkan");
 				PrivateDependencyModuleNames.Add("VulkanRHI");
