Fix console CIS

diff --git a/Engine/Plugins/Media/MediaIOFramework/Source/MediaIOCore/Private/Player/MediaIOCorePlayerBase.cpp b/Engine/Plugins/Media/MediaIOFramework/Source/MediaIOCore/Private/Player/MediaIOCorePlayerBase.cpp
--- a/Engine/Plugins/Media/MediaIOFramework/Source/MediaIOCore/Private/Player/MediaIOCorePlayerBase.cpp
+++ b/Engine/Plugins/Media/MediaIOFramework/Source/MediaIOCore/Private/Player/MediaIOCorePlayerBase.cpp
@@ -20,6 +20,8 @@
 #include "TimedDataInputCollection.h"
 #include "TimeSynchronizableMediaSource.h"
 
+#include <cstdint>
+
 #define LOCTEXT_NAMESPACE "MediaIOCorePlayerBase"
 
 /* MediaIOCorePlayerDetail
