Fix console CIS

diff --git a/Engine/Plugins/Media/MediaIOFramework/Source/GPUTextureTransfer/Private/TextureTransferBase.cpp b/Engine/Plugins/Media/MediaIOFramework/Source/GPUTextureTransfer/Private/TextureTransferBase.cpp
--- a/Engine/Plugins/Media/MediaIOFramework/Source/GPUTextureTransfer/Private/TextureTransferBase.cpp
+++ b/Engine/Plugins/Media/MediaIOFramework/Source/GPUTextureTransfer/Private/TextureTransferBase.cpp
@@ -160,13 +160,13 @@ namespace UE::GPUTextureTransfer::Private
 
 		if (!BufferInfo)
 		{
-			UE_LOG(LogGPUTextureTransfer, Error, TEXT("Error while performing a GPU transfer texture, CPU Buffer %u was not registered."), reinterpret_cast<std::uintptr_t>(InBuffer));
+			UE_LOG(LogGPUTextureTransfer, Error, TEXT("Error while performing a GPU transfer texture, CPU Buffer %u was not registered."), reinterpret_cast<uintptr_t>(InBuffer));
 			return false;
 		}
 
 		if (!TextureInfo)
 		{
-			UE_LOG(LogGPUTextureTransfer, Error, TEXT("Error while performing a GPU transfer texture, texture %u was not registered."), reinterpret_cast<std::uintptr_t>(InRHITexture));
+			UE_LOG(LogGPUTextureTransfer, Error, TEXT("Error while performing a GPU transfer texture, texture %u was not registered."), reinterpret_cast<uintptr_t>(InRHITexture));
 			return false;
 		}
 
diff --git a/Engine/Plugins/Media/MediaIOFramework/Source/MediaIOCore/Private/Player/MediaIOCorePlayerBase.cpp b/Engine/Plugins/Media/MediaIOFramework/Source/MediaIOCore/Private/Player/MediaIOCorePlayerBase.cpp
--- a/Engine/Plugins/Media/MediaIOFramework/Source/MediaIOCore/Private/Player/MediaIOCorePlayerBase.cpp
+++ b/Engine/Plugins/Media/MediaIOFramework/Source/MediaIOCore/Private/Player/MediaIOCorePlayerBase.cpp
@@ -20,8 +20,6 @@
 #include "TimedDataInputCollection.h"
 #include "TimeSynchronizableMediaSource.h"
 
-#include <cstdint>
-
 #define LOCTEXT_NAMESPACE "MediaIOCorePlayerBase"
 
 /* MediaIOCorePlayerDetail
@@ -563,7 +561,7 @@ void FMediaIOCorePlayerBase::RegisterSampleBuffer(const TSharedPtr<FMediaIOCoreT
 					TextureStride *= 4;
 				}
 
-				UE_LOG(LogMediaIOCore, Verbose, TEXT("Registering buffer %u"), reinterpret_cast<std::uintptr_t>(InSample->GetMutableBuffer()));
+				UE_LOG(LogMediaIOCore, Verbose, TEXT("Registering buffer %u"), reinterpret_cast<uintptr_t>(InSample->GetMutableBuffer()));
 				UE::GPUTextureTransfer::FRegisterDMABufferArgs Args;
 				Args.Buffer = InSample->GetMutableBuffer();
 				Args.Width = TextureWidth;
@@ -575,7 +573,7 @@ void FMediaIOCorePlayerBase::RegisterSampleBuffer(const TSharedPtr<FMediaIOCoreT
 	}
 	else
 	{
-		UE_LOG(LogMediaIOCore, VeryVerbose, TEXT("Buffer already registered: %u"), reinterpret_cast<std::uintptr_t>(InSample->GetMutableBuffer()));
+		UE_LOG(LogMediaIOCore, VeryVerbose, TEXT("Buffer already registered: %u"), reinterpret_cast<uintptr_t>(InSample->GetMutableBuffer()));
 	}
 }
 
@@ -647,7 +645,7 @@ void FMediaIOCorePlayerBase::CreateAndRegisterTextures(const IMediaOptions* Opti
 				FRHITextureCreateDesc CreateDesc = FRHITextureCreateDesc::Create2D(*TextureName, FIntPoint(TextureWidth, TextureHeight), InputFormat, FClearValueBinding::White, CreateFlags);
 				RHICreateTargetableShaderResource(CreateDesc, TexCreate_RenderTargetable, RHITexture, DummyTexture2DRHI);
 
-				UE_LOG(LogMediaIOCore, Verbose, TEXT("Registering texture %u"), reinterpret_cast<std::uintptr_t>(RHITexture->GetNativeResource()));
+				UE_LOG(LogMediaIOCore, Verbose, TEXT("Registering texture %u"), reinterpret_cast<uintptr_t>(RHITexture->GetNativeResource()));
 
 				UE::GPUTextureTransfer::FRegisterDMATextureArgs Args;
 				Args.RHITexture = RHITexture.GetReference();
@@ -681,7 +679,7 @@ void FMediaIOCorePlayerBase::UnregisterTextures()
 			{
 				for (const TRefCountPtr<FRHITexture>& RHITexture : TexturesToUnregister)
 				{
-					UE_LOG(LogMediaIOCore, Verbose, TEXT("Unregistering texture %u"), reinterpret_cast<std::uintptr_t>(RHITexture->GetNativeResource()));
+					UE_LOG(LogMediaIOCore, Verbose, TEXT("Unregistering texture %u"), reinterpret_cast<uintptr_t>(RHITexture->GetNativeResource()));
 					TextureTransfer->UnregisterTexture(RHITexture.GetReference());
 				}
 			}
@@ -714,7 +712,7 @@ void FMediaIOCorePlayerBase::PreGPUTransfer(const TSharedPtr<FMediaIOCoreTexture
 	}
 	else
 	{
-		UE_LOG(LogMediaIOCore, Display, TEXT("Unregistered texture %u encountered while doing a gpu texture transfer."), Texture ? reinterpret_cast<std::uintptr_t>(Texture->GetNativeResource()) : 0);
+		UE_LOG(LogMediaIOCore, Display, TEXT("Unregistered texture %u encountered while doing a gpu texture transfer."), Texture ? reinterpret_cast<uintptr_t>(Texture->GetNativeResource()) : 0);
 	}
 }
 
@@ -729,7 +727,7 @@ void FMediaIOCorePlayerBase::ExecuteGPUTransfer(const TSharedPtr<FMediaIOCoreTex
 				MediaPlayerPtr != nullptr && MediaPlayerPtr->GPUTextureTransfer != nullptr)
 			{
 				void* Buffer = InSample->GetMutableBuffer();
-				UE_LOG(LogMediaIOCore, Verbose, TEXT("Starting Transfer with buffer %u"), reinterpret_cast<std::uintptr_t>(Buffer));
+				UE_LOG(LogMediaIOCore, Verbose, TEXT("Starting Transfer with buffer %u"), reinterpret_cast<uintptr_t>(Buffer));
 				MediaPlayerPtr->GPUTextureTransfer->TransferTexture(Buffer, InSample->GetTexture(), UE::GPUTextureTransfer::ETransferDirection::CPU_TO_GPU);
 				MediaPlayerPtr->GPUTextureTransfer->BeginSync((InSample)->GetMutableBuffer(), UE::GPUTextureTransfer::ETransferDirection::CPU_TO_GPU);
 				MediaPlayerPtr->AddVideoSample(InSample.ToSharedRef());
