- Updating MediaCapture to use RDG
- Adding structured buffer output support to media capture

- Adding Rivermax output support wiht custom conversions to support

YCbCr10
RGB8
RGB10
RGB16f

- Fixing chunk configuration for output stream
- Fixing crash when dragging a rivermax player in the scene

diff --git a/Engine/Plugins/Media/AjaMedia/Source/AjaMediaOutput/Private/AjaMediaCapture.cpp b/Engine/Plugins/Media/AjaMedia/Source/AjaMediaOutput/Private/AjaMediaCapture.cpp
--- a/Engine/Plugins/Media/AjaMedia/Source/AjaMediaOutput/Private/AjaMediaCapture.cpp
+++ b/Engine/Plugins/Media/AjaMedia/Source/AjaMediaOutput/Private/AjaMediaCapture.cpp
@@ -287,7 +287,7 @@ void UAjaMediaCapture::StopCaptureImpl(bool bAllowPendingFrameToBeProcess)
 			[this](FRHICommandListImmediate& RHICmdList) mutable
 			{
 				// Unregister texture before closing channel.
-				if (ShouldCaptureRHITexture())
+				if (ShouldCaptureRHIResource())
 				{
 					for (FTextureRHIRef& Texture : TexturesToRelease)
 					{
@@ -320,14 +320,14 @@ void UAjaMediaCapture::StopCaptureImpl(bool bAllowPendingFrameToBeProcess)
 	}
 }
 
-bool UAjaMediaCapture::ShouldCaptureRHITexture() const
+bool UAjaMediaCapture::ShouldCaptureRHIResource() const
 {
 	return bGPUTextureTransferAvailable && CVarAjaEnableGPUDirect.GetValueOnAnyThread() == 1;
 }
 
 void UAjaMediaCapture::BeforeFrameCaptured_RenderingThread(const FCaptureBaseData& InBaseData, TSharedPtr<FMediaCaptureUserData, ESPMode::ThreadSafe> InUserData, FTextureRHIRef InTexture)
 {
-	if (ShouldCaptureRHITexture())
+	if (ShouldCaptureRHIResource())
 	{
 		if (!TexturesToRelease.Contains(InTexture))
 		{
@@ -406,7 +406,7 @@ bool UAjaMediaCapture::InitAJA(UAjaMediaOutput* InAjaMediaOutput)
 	FrameRate = InAjaMediaOutput->GetRequestedFrameRate();
 	PortName = FAjaDeviceProvider().ToText(InAjaMediaOutput->OutputConfiguration.MediaConfiguration.MediaConnection).ToString();
 
-	if (ShouldCaptureRHITexture())
+	if (ShouldCaptureRHIResource())
 	{
 		if (InAjaMediaOutput->OutputConfiguration.MediaConfiguration.MediaMode.Standard != EMediaIOStandardType::Progressive)
 		{
@@ -450,7 +450,7 @@ bool UAjaMediaCapture::InitAJA(UAjaMediaOutput* InAjaMediaOutput)
 	ChannelOptions.PixelFormat = AjaMediaCaptureUtils::ConvertPixelFormat(InAjaMediaOutput->PixelFormat, ChannelOptions.bUseKey);
 	ChannelOptions.TimecodeFormat =  AjaMediaCaptureUtils::ConvertTimecode(InAjaMediaOutput->TimecodeFormat);
 	ChannelOptions.OutputReferenceType = AjaMediaCaptureUtils::Convert(InAjaMediaOutput->OutputConfiguration.OutputReference);
-	ChannelOptions.bUseGPUDMA = ShouldCaptureRHITexture();
+	ChannelOptions.bUseGPUDMA = ShouldCaptureRHIResource();
 
 	bOutputAudio = InAjaMediaOutput->bOutputAudio;
 	
@@ -489,7 +489,7 @@ bool UAjaMediaCapture::InitAJA(UAjaMediaOutput* InAjaMediaOutput)
 		WakeUpEvent = FPlatformProcess::GetSynchEventFromPool(bIsManualReset);
 	}
 
-	if (ShouldCaptureRHITexture())
+	if (ShouldCaptureRHIResource())
 	{
 		UE_LOG(LogAjaMediaOutput, Display, TEXT("Aja capture started using GPU Direct"));
 	}
@@ -546,7 +546,7 @@ void UAjaMediaCapture::OnFrameCaptured_RenderingThread(const FCaptureBaseData& I
 	}
 }
 
-void UAjaMediaCapture::OnRHITextureCaptured_RenderingThread(const FCaptureBaseData& InBaseData,	TSharedPtr<FMediaCaptureUserData, ESPMode::ThreadSafe> InUserData, FTextureRHIRef InTexture)
+void UAjaMediaCapture::OnRHIResourceCaptured_RenderingThread(const FCaptureBaseData& InBaseData,	TSharedPtr<FMediaCaptureUserData, ESPMode::ThreadSafe> InUserData, FTextureRHIRef InTexture)
 {
 	TRACE_CPUPROFILER_EVENT_SCOPE(UAjaMediaCapture::OnFrameCaptured_RenderingThread);
 	
diff --git a/Engine/Plugins/Media/AjaMedia/Source/AjaMediaOutput/Public/AjaMediaCapture.h b/Engine/Plugins/Media/AjaMedia/Source/AjaMediaOutput/Public/AjaMediaCapture.h
--- a/Engine/Plugins/Media/AjaMedia/Source/AjaMediaOutput/Public/AjaMediaCapture.h
+++ b/Engine/Plugins/Media/AjaMedia/Source/AjaMediaOutput/Public/AjaMediaCapture.h
@@ -41,11 +41,11 @@ protected:
 	virtual bool UpdateSceneViewportImpl(TSharedPtr<FSceneViewport>& InSceneViewport) override;
 	virtual bool UpdateRenderTargetImpl(UTextureRenderTarget2D* InRenderTarget) override;
 	virtual void StopCaptureImpl(bool bAllowPendingFrameToBeProcess) override;
-	virtual bool ShouldCaptureRHITexture() const override;
+	virtual bool ShouldCaptureRHIResource() const override;
 	
 	virtual void BeforeFrameCaptured_RenderingThread(const FCaptureBaseData& InBaseData, TSharedPtr<FMediaCaptureUserData, ESPMode::ThreadSafe> InUserData, FTextureRHIRef InTexture) override;
 	virtual void OnFrameCaptured_RenderingThread(const FCaptureBaseData& InBaseData, TSharedPtr<FMediaCaptureUserData, ESPMode::ThreadSafe> InUserData, void* InBuffer, int32 Width, int32 Height, int32 BytesPerRow) override;
-	virtual void OnRHITextureCaptured_RenderingThread(const FCaptureBaseData& InBaseData, TSharedPtr<FMediaCaptureUserData, ESPMode::ThreadSafe> InUserData, FTextureRHIRef InTexture) override;
+	virtual void OnRHIResourceCaptured_RenderingThread(const FCaptureBaseData& InBaseData, TSharedPtr<FMediaCaptureUserData, ESPMode::ThreadSafe> InUserData, FTextureRHIRef InTexture) override;
 
 private:
 	struct FAjaOutputCallback;
diff --git a/Engine/Plugins/Media/BlackmagicMedia/Source/BlackmagicMediaOutput/Private/BlackmagicMediaCapture.cpp b/Engine/Plugins/Media/BlackmagicMedia/Source/BlackmagicMediaOutput/Private/BlackmagicMediaCapture.cpp
--- a/Engine/Plugins/Media/BlackmagicMedia/Source/BlackmagicMediaOutput/Private/BlackmagicMediaCapture.cpp
+++ b/Engine/Plugins/Media/BlackmagicMedia/Source/BlackmagicMediaOutput/Private/BlackmagicMediaCapture.cpp
@@ -423,7 +423,7 @@ void UBlackmagicMediaCapture::StopCaptureImpl(bool bAllowPendingFrameToBeProcess
 	}
 }
 
-bool UBlackmagicMediaCapture::ShouldCaptureRHITexture() const
+bool UBlackmagicMediaCapture::ShouldCaptureRHIResource() const
 {
 	// Todo: also test if dvp was initialized correctly.
 	return bGPUTextureTransferAvailable && CVarBlackmagicEnableGPUDirect.GetValueOnAnyThread() == 1;
@@ -490,7 +490,7 @@ bool UBlackmagicMediaCapture::InitBlackmagic(UBlackmagicMediaOutput* InBlackmagi
 	bLogDropFrame = InBlackmagicMediaOutput->bLogDropFrame;
 	FrameRate = InBlackmagicMediaOutput->GetRequestedFrameRate();
 
-	if (ShouldCaptureRHITexture())
+	if (ShouldCaptureRHIResource())
 	{
 		if (InBlackmagicMediaOutput->OutputConfiguration.MediaConfiguration.MediaMode.Standard != EMediaIOStandardType::Progressive)
 		{
@@ -529,7 +529,7 @@ bool UBlackmagicMediaCapture::InitBlackmagic(UBlackmagicMediaOutput* InBlackmagi
 	ChannelOptions.bOutputVideo = true;
 	ChannelOptions.bOutputInterlacedFieldsTimecodeNeedToMatch = InBlackmagicMediaOutput->bInterlacedFieldsTimecodeNeedToMatch && InBlackmagicMediaOutput->OutputConfiguration.MediaConfiguration.MediaMode.Standard == EMediaIOStandardType::Interlaced && InBlackmagicMediaOutput->TimecodeFormat != EMediaIOTimecodeFormat::None;
 	ChannelOptions.bLogDropFrames = bLogDropFrame;
-	ChannelOptions.bUseGPUDMA = ShouldCaptureRHITexture();
+	ChannelOptions.bUseGPUDMA = ShouldCaptureRHIResource();
 	ChannelOptions.bScheduleInDifferentThread = InBlackmagicMediaOutput->bUseMultithreadedScheduling;
 
 	AudioBitDepth = InBlackmagicMediaOutput->AudioBitDepth;
@@ -573,7 +573,7 @@ bool UBlackmagicMediaCapture::InitBlackmagic(UBlackmagicMediaOutput* InBlackmagi
 		WakeUpEvent = FPlatformProcess::GetSynchEventFromPool(bIsManualReset);
 	}
 
-	if (ShouldCaptureRHITexture())
+	if (ShouldCaptureRHIResource())
 	{
 		UE_LOG(LogBlackmagicMediaOutput, Display, TEXT("BlackmagicMedia capture started using GPU Direct"));
 	}
@@ -583,7 +583,7 @@ bool UBlackmagicMediaCapture::InitBlackmagic(UBlackmagicMediaOutput* InBlackmagi
 
 void UBlackmagicMediaCapture::BeforeFrameCaptured_RenderingThread(const FCaptureBaseData& InBaseData, TSharedPtr<FMediaCaptureUserData, ESPMode::ThreadSafe> InUserData, FTextureRHIRef InTexture)
 {
-	if (ShouldCaptureRHITexture())
+	if (ShouldCaptureRHIResource())
 	{
 		if (!TexturesToRelease.Contains(InTexture))
 		{
@@ -678,7 +678,7 @@ void UBlackmagicMediaCapture::OnFrameCaptured_RenderingThread(const FCaptureBase
 	}
 }
 
-void UBlackmagicMediaCapture::OnRHITextureCaptured_RenderingThread(const FCaptureBaseData& InBaseData, TSharedPtr<FMediaCaptureUserData, ESPMode::ThreadSafe> InUserData, FTextureRHIRef InTexture)
+void UBlackmagicMediaCapture::OnRHIResourceCaptured_RenderingThread(const FCaptureBaseData& InBaseData, TSharedPtr<FMediaCaptureUserData, ESPMode::ThreadSafe> InUserData, FTextureRHIRef InTexture)
 {
 	if (!InTexture)
 	{
diff --git a/Engine/Plugins/Media/BlackmagicMedia/Source/BlackmagicMediaOutput/Public/BlackmagicMediaCapture.h b/Engine/Plugins/Media/BlackmagicMedia/Source/BlackmagicMediaOutput/Public/BlackmagicMediaCapture.h
--- a/Engine/Plugins/Media/BlackmagicMedia/Source/BlackmagicMediaOutput/Public/BlackmagicMediaCapture.h
+++ b/Engine/Plugins/Media/BlackmagicMedia/Source/BlackmagicMediaOutput/Public/BlackmagicMediaCapture.h
@@ -41,11 +41,11 @@ protected:
 	virtual bool UpdateSceneViewportImpl(TSharedPtr<FSceneViewport>& InSceneViewport) override;
 	virtual bool UpdateRenderTargetImpl(UTextureRenderTarget2D* InRenderTarget) override;
 	virtual void StopCaptureImpl(bool bAllowPendingFrameToBeProcess) override;
-	virtual bool ShouldCaptureRHITexture() const override;
+	virtual bool ShouldCaptureRHIResource() const override;
 
 	virtual void BeforeFrameCaptured_RenderingThread(const FCaptureBaseData& InBaseData, TSharedPtr<FMediaCaptureUserData, ESPMode::ThreadSafe> InUserData, FTextureRHIRef InTexture) override;
 	virtual void OnFrameCaptured_RenderingThread(const FCaptureBaseData& InBaseData, TSharedPtr<FMediaCaptureUserData, ESPMode::ThreadSafe> InUserData, void* InBuffer, int32 Width, int32 Height, int32 BytesPerRow) override;
-	virtual void OnRHITextureCaptured_RenderingThread(const FCaptureBaseData& InBaseData, TSharedPtr<FMediaCaptureUserData, ESPMode::ThreadSafe> InUserData, FTextureRHIRef InTexture) override;
+	virtual void OnRHIResourceCaptured_RenderingThread(const FCaptureBaseData& InBaseData, TSharedPtr<FMediaCaptureUserData, ESPMode::ThreadSafe> InUserData, FTextureRHIRef InTexture) override;
 
 private:
 	struct FBlackmagicOutputCallback;
diff --git a/Engine/Plugins/Media/MediaIOFramework/Source/MediaIOCore/MediaIOCore.Build.cs b/Engine/Plugins/Media/MediaIOFramework/Source/MediaIOCore/MediaIOCore.Build.cs
--- a/Engine/Plugins/Media/MediaIOFramework/Source/MediaIOCore/MediaIOCore.Build.cs
+++ b/Engine/Plugins/Media/MediaIOFramework/Source/MediaIOCore/MediaIOCore.Build.cs
@@ -1,5 +1,7 @@
 // Copyright Epic Games, Inc. All Rights Reserved.
 
+using System.IO;
+
 namespace UnrealBuildTool.Rules
 {
 	public class MediaIOCore : ModuleRules
@@ -25,11 +27,14 @@ namespace UnrealBuildTool.Rules
 					"SlateCore",
 					"TimeManagement"
 				});
-			
+
+			var EngineDir = Path.GetFullPath(Target.RelativeEnginePath);
 			PrivateIncludePaths.AddRange(
 				new string[] {
+					//required for FScreenPassVS and AddDrawScreenPass
+					Path.Combine(EngineDir, "Source", "Runtime", "Renderer", "Private")
 				});
-			
+
 			PrivateDependencyModuleNames.AddRange(
                 new string[]
                 {
@@ -37,6 +42,7 @@ namespace UnrealBuildTool.Rules
 	                "AudioMixerCore",
 					"GPUTextureTransfer",
                     "RenderCore",
+					"Renderer",
                     "SignalProcessing", 
 	                "SoundFieldRendering"
 				});
diff --git a/Engine/Plugins/Media/MediaIOFramework/Source/MediaIOCore/Private/MediaCapture.cpp b/Engine/Plugins/Media/MediaIOFramework/Source/MediaIOCore/Private/MediaCapture.cpp
--- a/Engine/Plugins/Media/MediaIOFramework/Source/MediaIOCore/Private/MediaCapture.cpp
+++ b/Engine/Plugins/Media/MediaIOFramework/Source/MediaIOCore/Private/MediaCapture.cpp
@@ -16,35 +16,36 @@
 #include "Misc/ScopeLock.h"
 #include "PipelineStateCache.h"
 #include "RendererInterface.h"
+#include "RenderGraphResources.h"
+#include "RenderGraphUtils.h"
 #include "RenderUtils.h"
+#include "RenderTargetPool.h"
 #include "RHIStaticStates.h"
 #include "SceneUtils.h"
+#include "ScreenPass.h"
 #include "Slate/SceneViewport.h"
 #include "UObject/WeakObjectPtrTemplates.h"
-#include "RenderTargetPool.h"
 
 #if WITH_EDITOR
 #include "AnalyticsEventAttribute.h"
 #include "Editor.h"
-#include "EngineAnalytics.h"
-#include "IAssetViewport.h"
-#include "LevelEditor.h"
 #include "Editor/EditorEngine.h"
 #include "Engine/Engine.h"
 #include "Engine/EngineTypes.h"
+#include "EngineAnalytics.h"
 #include "Framework/Notifications/NotificationManager.h"
+#include "IAssetViewport.h"
+#include "LevelEditor.h"
 #include "Widgets/Notifications/SNotificationList.h"
 #endif //WITH_EDITOR
 
 #define LOCTEXT_NAMESPACE "MediaCapture"
 
-/* namespace MediaCaptureDetails definition
-*****************************************************************************/
 
 /** Time spent in media capture sending a frame. */
-DECLARE_CYCLE_STAT(TEXT("MediaCapture RenderThread CopyToResolve"), STAT_MediaCapture_RenderThread_CopyToResolve, STATGROUP_Media);
-DECLARE_CYCLE_STAT(TEXT("MediaCapture RenderThread MapStaging"), STAT_MediaCapture_RenderThread_MapStaging, STATGROUP_Media);
-DECLARE_CYCLE_STAT(TEXT("MediaCapture RenderThread Callback"), STAT_MediaCapture_RenderThread_Callback, STATGROUP_Media);
+DECLARE_CYCLE_STAT(TEXT("MediaCapture RenderThread FrameCapture"), STAT_MediaCapture_RenderThread_FrameCapture, STATGROUP_Media);
+DECLARE_CYCLE_STAT(TEXT("MediaCapture RenderThread LockResource"), STAT_MediaCapture_RenderThread_LockResource, STATGROUP_Media);
+DECLARE_CYCLE_STAT(TEXT("MediaCapture RenderThread RHI Capture Callback"), STAT_MediaCapture_RenderThread_CaptureCallback, STATGROUP_Media);
 
 /** These pixel formats do not require additional conversion except for swizzling and normalized sampling. */
 static TSet<EPixelFormat> SupportedRgbaSwizzleFormats =
@@ -85,6 +86,521 @@ static TSet<EPixelFormat> SupportedRgbaSwizzleFormats =
 	PF_R8,
 };
 
+namespace UE::MediaCaptureData
+{
+	class FCaptureFrame
+	{
+	public:
+		FCaptureFrame()
+			: bReadbackRequested(false)
+		{
+
+		}
+		virtual ~FCaptureFrame() {};
+
+		/** Returns true if its output resource is valid */
+		virtual bool HasValidResource() const = 0;
+
+		/** Simple way to validate the resource type and cast safely */
+		virtual bool IsTextureResource() const = 0;
+		virtual bool IsBufferResource() const = 0;
+
+		/** Locks the readback resource and returns a pointer to access data from system memory */
+		virtual void* Lock(FRHICommandListImmediate& RHICmdList, FRHIGPUMask GPUMask, int32& OutWidth) = 0;
+
+		/** Unlocks the readback resource */
+		virtual void Unlock() = 0;
+
+		/** Returns true if the readback is ready to be used */
+		virtual bool IsReadbackReady(FRHIGPUMask GPUMask) = 0;
+
+		UMediaCapture::FCaptureBaseData CaptureBaseData;
+		TAtomic<bool> bReadbackRequested;
+		TSharedPtr<FMediaCaptureUserData> UserData;
+	};
+
+	class FTextureCaptureFrame : public FCaptureFrame
+	{
+	public:
+		/** Type alias for the output resource type used during capture frame */
+		using FOutputResourceType = FRDGTextureRef;
+
+		//~ Begin FCaptureFrame interface
+		virtual bool HasValidResource() const override
+		{
+			return RenderTarget != nullptr;
+		}
+
+		virtual bool IsTextureResource() const override
+		{
+			return true;
+		}
+
+		virtual bool IsBufferResource() const override
+		{
+			return false;
+		}
+
+		virtual void* Lock(FRHICommandListImmediate& RHICmdList, FRHIGPUMask GPUMask, int32& OutWidth) override
+		{
+			if (ReadbackTexture->IsReady(GPUMask) == false)
+			{
+				UE_LOG(LogMediaIOCore, Verbose, TEXT("Fenced for texture readback was not ready"));
+			}
+
+			return ReadbackTexture->Lock(OutWidth);
+		}
+
+		virtual void Unlock() override
+		{
+			ReadbackTexture->Unlock();
+		}
+
+		virtual bool IsReadbackReady(FRHIGPUMask GPUMask) override
+		{
+			return ReadbackTexture->IsReady(GPUMask);
+		}
+
+		//~ End FCaptureFrame interface
+
+		/** Registers an external texture to be tracked by the graph and returns a pointer to the tracked resource */
+		FRDGTextureRef RegisterResource(FRDGBuilder& RDGBuilder)
+		{
+			return RDGBuilder.RegisterExternalTexture(RenderTarget, TEXT("OutputTexture"));
+		}
+
+		/** Adds a readback pass to the graph */
+		void EnqueueCopy(FRDGBuilder& RDGBuilder, FRDGTextureRef ResourceToReadback)
+		{
+			AddEnqueueCopyPass(RDGBuilder, ReadbackTexture.Get(), ResourceToReadback);
+		}
+
+		/** Returns RHI resource of the allocated pooled resource */
+		FRHITexture* GetRHIResource()
+		{
+			return RenderTarget->GetRHI();
+		}
+		
+	public:
+		TRefCountPtr<IPooledRenderTarget> RenderTarget;
+		TUniquePtr<FRHIGPUTextureReadback> ReadbackTexture;
+	};
+
+	class FBufferCaptureFrame : public FCaptureFrame
+	{
+	public:
+		/** Type alias for the output resource type used during capture frame */
+		using FOutputResourceType = FRDGBufferRef;
+
+		//~ Begin FCaptureFrame interface
+		virtual bool HasValidResource() const override
+		{
+			return Buffer != nullptr;
+		}
+
+		virtual bool IsTextureResource() const override
+		{
+			return false;
+		}
+
+		virtual bool IsBufferResource() const override
+		{
+			return true;
+		}
+
+		virtual void* Lock(FRHICommandListImmediate& RHICmdList, FRHIGPUMask GPUMask, int32& OutWidth) override
+		{
+			if (ReadbackBuffer->IsReady(GPUMask) == false)
+			{
+				UE_LOG(LogMediaIOCore, Verbose, TEXT("Fenced for buffer readback was not ready, blocking."));
+				RHICmdList.BlockUntilGPUIdle();
+			}
+
+			OutWidth = Buffer->GetRHI()->GetStride();
+			return ReadbackBuffer->Lock(Buffer->GetRHI()->GetSize());
+		}
+
+		virtual void Unlock() override
+		{
+			ReadbackBuffer->Unlock();
+		}
+
+		virtual bool IsReadbackReady(FRHIGPUMask GPUMask) override
+		{
+			return ReadbackBuffer->IsReady(GPUMask);
+		}
+		//~ End FCaptureFrame interface
+
+		/** Registers an external texture to be tracked by the graph and returns a pointer to the tracked resource */
+		FRDGBufferRef RegisterResource(FRDGBuilder& RDGBuilder)
+		{
+			return RDGBuilder.RegisterExternalBuffer(Buffer);
+		}
+
+		/** Adds a readback pass to the graph */
+		void EnqueueCopy(FRDGBuilder& RDGBuilder, FRDGBufferRef ResourceToReadback)
+		{
+			AddEnqueueCopyPass(RDGBuilder, ReadbackBuffer.Get(), ResourceToReadback, Buffer->GetRHI()->GetSize());
+		}
+
+		/** Returns RHI resource of the allocated pooled resource */
+		FRHIBuffer* GetRHIResource()
+		{
+			return Buffer->GetRHI();
+		}
+
+	public:
+		TRefCountPtr<FRDGPooledBuffer> Buffer;
+		TUniquePtr<FRHIGPUBufferReadback> ReadbackBuffer;
+	};
+
+	/** Helper struct to contain arguments for CaptureFrame */
+	struct FCaptureFrameArgs
+	{
+		FRHICommandListImmediate& RHICmdList;
+		TObjectPtr<UMediaCapture> MediaCapture = nullptr;
+		FSceneViewport* CapturingSceneViewport = nullptr;
+		FTextureRenderTargetResource* TextureRenderTargetResource = nullptr;
+		FIntPoint DesiredSize = FIntPoint::ZeroValue;
+	};
+
+	/** Helper struct to contain arguments for AddConversionPass */
+	struct FConversionPassArgs
+	{
+		FRDGBuilder& GraphBuilder;
+		const FRDGTextureRef& SourceRGBTexture;
+		bool bRequiresFormatConversion = false;
+		FResolveParams ResolveParams;
+		FVector2D SizeU = FVector2D::ZeroVector;
+		FVector2D SizeV = FVector2D::ZeroVector;
+	};
+
+	/** Helper class to be able to friend it and call methods on input media capture */
+	class FMediaCaptureHelper
+	{
+	public:
+
+		static void BeforeFrameCapture(const FCaptureFrameArgs& Args, FTextureCaptureFrame* CapturingFrame)
+		{
+			UE_LOG(LogMediaIOCore, VeryVerbose, TEXT("Before Locking texture %p"), reinterpret_cast<uintptr_t>(CapturingFrame->RenderTarget->GetRHI()->GetNativeResource()));
+			Args.MediaCapture->BeforeFrameCaptured_RenderingThread(CapturingFrame->CaptureBaseData, CapturingFrame->UserData, CapturingFrame->RenderTarget->GetRHI());
+			UE_LOG(LogMediaIOCore, VeryVerbose, TEXT("After Locking texture %p"), reinterpret_cast<uintptr_t>(CapturingFrame->RenderTarget->GetRHI()->GetNativeResource()));
+		}
+
+		static void BeforeFrameCapture(const FCaptureFrameArgs& Args, FBufferCaptureFrame* CapturingFrame)
+		{
+			FBufferCaptureFrame* BufferFrame = static_cast<FBufferCaptureFrame*>(CapturingFrame);
+			Args.MediaCapture->BeforeFrameCaptured_RenderingThread(CapturingFrame->CaptureBaseData, CapturingFrame->UserData, BufferFrame->Buffer->GetRHI());
+		}
+
+		static FTexture2DRHIRef GetSourceTextureForInput(const FCaptureFrameArgs& Args)
+		{
+			FTexture2DRHIRef SourceTexture;
+			if (Args.CapturingSceneViewport)
+			{
+#if WITH_EDITOR
+				if (!IsRunningGame())
+				{
+					// PIE, PIE in windows, editor viewport
+					SourceTexture = Args.CapturingSceneViewport->GetRenderTargetTexture();
+					if (!SourceTexture.IsValid() && Args.CapturingSceneViewport->GetViewportRHI())
+					{
+						SourceTexture = Args.RHICmdList.GetViewportBackBuffer(Args.CapturingSceneViewport->GetViewportRHI());
+					}
+				}
+				else
+#endif
+					if (Args.CapturingSceneViewport->GetViewportRHI())
+					{
+						// Standalone and packaged
+						SourceTexture = Args.RHICmdList.GetViewportBackBuffer(Args.CapturingSceneViewport->GetViewportRHI());
+					}
+			}
+			else if (Args.TextureRenderTargetResource && Args.TextureRenderTargetResource->GetTextureRenderTarget2DResource())
+			{
+				SourceTexture = Args.TextureRenderTargetResource->GetTextureRenderTarget2DResource()->GetTextureRHI();
+			}
+
+			return SourceTexture;
+		}
+
+		static bool AreInputsValid(const FCaptureFrameArgs& Args, const FTexture2DRHIRef& SourceTexture)
+		{
+			// If it is a simple rgba swizzle we can handle the conversion. Supported formats
+			// contained in SupportedRgbaSwizzleFormats. Warning would've been displayed on start of capture.
+			if (Args.MediaCapture->DesiredPixelFormat != SourceTexture->GetFormat() &&
+				(!SupportedRgbaSwizzleFormats.Contains(SourceTexture->GetFormat()) || !Args.MediaCapture->DesiredCaptureOptions.bConvertToDesiredPixelFormat))
+			{
+				UE_LOG(LogMediaIOCore, Error, TEXT("The capture will stop for '%s'. The Source pixel format doesn't match with the user requested pixel format. %sRequested: %s Source: %s")
+					, *Args.MediaCapture->MediaOutputName
+					, (SupportedRgbaSwizzleFormats.Contains(SourceTexture->GetFormat()) && !Args.MediaCapture->DesiredCaptureOptions.bConvertToDesiredPixelFormat) ? TEXT("Please enable \"Convert To Desired Pixel Format\" option in Media Capture settings. ") : TEXT("")
+					, GetPixelFormatString(Args.MediaCapture->DesiredPixelFormat)
+					, GetPixelFormatString(SourceTexture->GetFormat()));
+
+				return false;
+			}
+
+			if (Args.MediaCapture->DesiredCaptureOptions.Crop == EMediaCaptureCroppingType::None)
+			{
+				if (Args.DesiredSize.X != SourceTexture->GetSizeX() || Args.DesiredSize.Y != SourceTexture->GetSizeY())
+				{
+					UE_LOG(LogMediaIOCore, Error, TEXT("The capture will stop for '%s'. The Source size doesn't match with the user requested size. Requested: %d,%d  Source: %d,%d")
+						, *Args.MediaCapture->MediaOutputName
+						, Args.DesiredSize.X, Args.DesiredSize.Y
+						, SourceTexture->GetSizeX(), SourceTexture->GetSizeY());
+
+					return false;
+				}
+			}
+			else
+			{
+				FIntPoint StartCapturePoint = FIntPoint::ZeroValue;
+				if (Args.MediaCapture->DesiredCaptureOptions.Crop == EMediaCaptureCroppingType::Custom)
+				{
+					StartCapturePoint = Args.MediaCapture->DesiredCaptureOptions.CustomCapturePoint;
+				}
+
+				if ((uint32)(Args.DesiredSize.X + StartCapturePoint.X) > SourceTexture->GetSizeX() || (uint32)(Args.DesiredSize.Y + StartCapturePoint.Y) > SourceTexture->GetSizeY())
+				{
+					UE_LOG(LogMediaIOCore, Error, TEXT("The capture will stop for '%s'. The Source size doesn't match with the user requested size. Requested: %d,%d  Source: %d,%d")
+						, *Args.MediaCapture->MediaOutputName
+						, Args.DesiredSize.X, Args.DesiredSize.Y
+						, SourceTexture->GetSizeX(), SourceTexture->GetSizeY());
+
+					return false;
+				}
+			}
+
+			return true;
+		}
+
+		static void GetResolveParameters(const FCaptureFrameArgs& Args, const FTexture2DRHIRef& SourceTexture, FResolveParams& OutResolveParams, FVector2D& OutSizeU, FVector2D& OutSizeV)
+		{
+			// Default to no crop
+			OutSizeU = { 0.0f, 1.0f };
+			OutSizeV = { 0.0f, 1.0f };
+			OutResolveParams.Rect = FResolveRect(0, 0, SourceTexture->GetSizeX(), SourceTexture->GetSizeY());
+			if (Args.MediaCapture->DesiredCaptureOptions.Crop != EMediaCaptureCroppingType::None)
+			{
+				switch (Args.MediaCapture->DesiredCaptureOptions.Crop)
+				{
+				case EMediaCaptureCroppingType::Center:
+					OutResolveParams.Rect = FResolveRect((SourceTexture->GetSizeX() - Args.DesiredSize.X) / 2, (SourceTexture->GetSizeY() - Args.DesiredSize.Y) / 2, 0, 0);
+					OutResolveParams.Rect.X2 = OutResolveParams.Rect.X1 + Args.DesiredSize.X;
+					OutResolveParams.Rect.Y2 = OutResolveParams.Rect.Y1 + Args.DesiredSize.Y;
+					break;
+				case EMediaCaptureCroppingType::TopLeft:
+					OutResolveParams.Rect = FResolveRect(0, 0, Args.DesiredSize.X, Args.DesiredSize.Y);
+					break;
+				case EMediaCaptureCroppingType::Custom:
+					OutResolveParams.Rect = FResolveRect(Args.MediaCapture->DesiredCaptureOptions.CustomCapturePoint.X, Args.MediaCapture->DesiredCaptureOptions.CustomCapturePoint.Y, 0, 0);
+					OutResolveParams.Rect.X2 = OutResolveParams.Rect.X1 + Args.DesiredSize.X;
+					OutResolveParams.Rect.Y2 = OutResolveParams.Rect.Y1 + Args.DesiredSize.Y;
+					break;
+				}
+
+				OutResolveParams.DestRect.X1 = 0;
+				OutResolveParams.DestRect.X2 = Args.DesiredSize.X;
+				OutResolveParams.DestRect.Y1 = 0;
+				OutResolveParams.DestRect.Y2 = Args.DesiredSize.Y;
+
+				OutSizeU.X = (float)OutResolveParams.Rect.X1 / (float)SourceTexture->GetSizeX();
+				OutSizeU.Y = (float)OutResolveParams.Rect.X2 / (float)SourceTexture->GetSizeX();
+				OutSizeV.X = (float)OutResolveParams.Rect.Y1 / (float)SourceTexture->GetSizeY();
+				OutSizeV.Y = (float)OutResolveParams.Rect.Y2 / (float)SourceTexture->GetSizeY();
+			}
+		}
+
+		static void AddConversionPass(const FCaptureFrameArgs& Args, const FConversionPassArgs& ConversionPassArgs, FRDGTextureRef OutputResource)
+		{
+			//Based on conversion type, this might be changed
+			bool bRequiresFormatConversion = ConversionPassArgs.bRequiresFormatConversion;
+
+			// Rectangle area to use from source
+			const FIntRect ViewRect(ConversionPassArgs.ResolveParams.Rect.X1, ConversionPassArgs.ResolveParams.Rect.Y1, ConversionPassArgs.ResolveParams.Rect.X2, ConversionPassArgs.ResolveParams.Rect.Y2);
+
+			//Dummy ViewFamily/ViewInfo created to use built in Draw Screen/Texture Pass
+			FSceneViewFamily ViewFamily(FSceneViewFamily::ConstructionValues(nullptr, nullptr, FEngineShowFlags(ESFIM_Game))
+				.SetTime(FGameTime())
+				.SetGammaCorrection(1.0f));
+			FSceneViewInitOptions ViewInitOptions;
+			ViewInitOptions.ViewFamily = &ViewFamily;
+			ViewInitOptions.SetViewRectangle(ViewRect);
+			ViewInitOptions.ViewOrigin = FVector::ZeroVector;
+			ViewInitOptions.ViewRotationMatrix = FMatrix::Identity;
+			ViewInitOptions.ProjectionMatrix = FMatrix::Identity;
+			FViewInfo ViewInfo = FViewInfo(ViewInitOptions);
+
+			// If no conversion was required, go through a simple copy
+			if (Args.MediaCapture->ConversionOperation == EMediaCaptureConversionOperation::NONE && !bRequiresFormatConversion)
+			{
+				AddCopyToResolveTargetPass(ConversionPassArgs.GraphBuilder, ConversionPassArgs.SourceRGBTexture, OutputResource, ConversionPassArgs.ResolveParams);
+			}
+			else
+			{
+				//At some point we should support color conversion (ocio) but for now we push incoming texture as is
+				constexpr bool bDoLinearToSRGB = false;
+
+				FGlobalShaderMap* GlobalShaderMap = GetGlobalShaderMap(GMaxRHIFeatureLevel);
+				TShaderMapRef<FScreenPassVS> VertexShader(GlobalShaderMap);
+
+				switch (Args.MediaCapture->ConversionOperation)
+				{
+				case EMediaCaptureConversionOperation::RGBA8_TO_YUV_8BIT:
+				{
+					//Configure source/output viewport to get the right UV scaling from source texture to output texture
+					FScreenPassTextureViewport InputViewport(ConversionPassArgs.SourceRGBTexture, ViewRect);
+					FScreenPassTextureViewport OutputViewport(OutputResource);
+					TShaderMapRef<FRGB8toUYVY8ConvertPS> PixelShader(GlobalShaderMap);
+					FRGB8toUYVY8ConvertPS::FParameters* Parameters = PixelShader->AllocateAndSetParameters(ConversionPassArgs.GraphBuilder, ConversionPassArgs.SourceRGBTexture, MediaShaders::RgbToYuvRec709Scaled, MediaShaders::YUVOffset8bits, bDoLinearToSRGB, OutputResource);
+					AddDrawScreenPass(ConversionPassArgs.GraphBuilder, RDG_EVENT_NAME("RGBToUYVY 8 bit"), ViewInfo, OutputViewport, InputViewport, VertexShader, PixelShader, Parameters);
+				}
+				break;
+				case EMediaCaptureConversionOperation::RGB10_TO_YUVv210_10BIT:
+				{
+					//Configure source/output viewport to get the right UV scaling from source texture to output texture
+					const FIntPoint InExtent = FIntPoint((((ConversionPassArgs.SourceRGBTexture->Desc.Extent.X + 47) / 48) * 48), ConversionPassArgs.SourceRGBTexture->Desc.Extent.Y);;
+					FScreenPassTextureViewport InputViewport(ConversionPassArgs.SourceRGBTexture, ViewRect);
+					FScreenPassTextureViewport OutputViewport(OutputResource);
+					TShaderMapRef<FRGB10toYUVv210ConvertPS> PixelShader(GlobalShaderMap);
+					FRGB10toYUVv210ConvertPS::FParameters* Parameters = PixelShader->AllocateAndSetParameters(ConversionPassArgs.GraphBuilder, ConversionPassArgs.SourceRGBTexture, MediaShaders::RgbToYuvRec709Scaled, MediaShaders::YUVOffset10bits, bDoLinearToSRGB, OutputResource);
+					AddDrawScreenPass(ConversionPassArgs.GraphBuilder, RDG_EVENT_NAME("RGBToYUVv210"), ViewInfo, OutputViewport, InputViewport, VertexShader, PixelShader, Parameters);
+				}
+				break;
+				case EMediaCaptureConversionOperation::INVERT_ALPHA:
+					// fall through
+				case EMediaCaptureConversionOperation::SET_ALPHA_ONE:
+					// fall through
+				case EMediaCaptureConversionOperation::NONE:
+					bRequiresFormatConversion = true;
+				default:
+					if (bRequiresFormatConversion)
+					{
+						//Configure source/output viewport to get the right UV scaling from source texture to output texture
+						FScreenPassTextureViewport InputViewport(ConversionPassArgs.SourceRGBTexture, ViewRect);
+						FScreenPassTextureViewport OutputViewport(OutputResource);
+
+						// In cases where texture is converted from a format that doesn't have A channel, we want to force set it to 1.
+						EMediaCaptureConversionOperation MediaConversionOperation = Args.MediaCapture->DesiredCaptureOptions.bForceAlphaToOneOnConversion ? EMediaCaptureConversionOperation::SET_ALPHA_ONE : Args.MediaCapture->ConversionOperation;
+						FModifyAlphaSwizzleRgbaPS::FPermutationDomain PermutationVector;
+						PermutationVector.Set<FModifyAlphaSwizzleRgbaPS::FConversionOp>(static_cast<int32>(MediaConversionOperation));
+
+						TShaderMapRef<FModifyAlphaSwizzleRgbaPS> PixelShader(GlobalShaderMap, PermutationVector);
+						FModifyAlphaSwizzleRgbaPS::FParameters* Parameters = PixelShader->AllocateAndSetParameters(ConversionPassArgs.GraphBuilder, ConversionPassArgs.SourceRGBTexture, OutputResource);
+						AddDrawScreenPass(ConversionPassArgs.GraphBuilder, RDG_EVENT_NAME("MediaCaptureSwizzle"), ViewInfo, OutputViewport, InputViewport, VertexShader, PixelShader, Parameters);
+					}
+					break;
+				}
+			}
+		}
+
+		static void AddConversionPass(const FCaptureFrameArgs& Args, const FConversionPassArgs& ConversionPassArgs, FRDGBufferRef OutputResource)
+		{
+			/** Used at some point to deal with native buffer output resource type but there to validate compiled template */
+		}
+
+		template<typename TFrameType>
+		static bool CaptureFrame(const FCaptureFrameArgs& Args, TFrameType* CapturingFrame)
+		{
+			// Validate if we have a resources used to capture source texture
+			if (!CapturingFrame->HasValidResource())
+			{
+				UE_LOG(LogMediaIOCore, Error, TEXT("The capture will stop for '%s'. A capture frame had an invalid render resource."), *Args.MediaCapture->MediaOutputName);
+				return false;
+			}
+
+			// Get source texture that is going to be captured
+			const FTexture2DRHIRef SourceTexture = GetSourceTextureForInput(Args);
+
+			// If true, we will need to go through our different shader to convert from source format to out format (i.e RGB to YUV)
+			const bool bRequiresFormatConversion = Args.MediaCapture->DesiredPixelFormat != SourceTexture->GetFormat();
+
+			if (!SourceTexture.IsValid())
+			{
+				UE_LOG(LogMediaIOCore, Error, TEXT("Can't grab the Texture to capture for '%s'."), *Args.MediaCapture->MediaOutputName);
+				return false;
+			}
+
+			// Validate pixel formats and sizes before pursuing
+			if (AreInputsValid(Args, SourceTexture) == false)
+			{
+				return false;
+			}
+
+			{
+				TRACE_CPUPROFILER_EVENT_SCOPE(UMediaCapture::BeforeFrameCaptured);
+				BeforeFrameCapture(Args, CapturingFrame);
+			}
+
+			{
+				TRACE_CPUPROFILER_EVENT_SCOPE(UMediaCapture::FrameCapture);
+				SCOPE_CYCLE_COUNTER(STAT_MediaCapture_RenderThread_FrameCapture);
+
+				FResolveParams ResolveParams;
+				FVector2D SizeU;
+				FVector2D SizeV;
+				GetResolveParameters(Args, SourceTexture, ResolveParams, SizeU, SizeV);
+
+				FRDGBuilder GraphBuilder(Args.RHICmdList);
+
+				// Register the source texture that we want to capture
+				const FRDGTextureRef SourceRGBTexture = GraphBuilder.RegisterExternalTexture(CreateRenderTarget(SourceTexture, TEXT("SourceTexture")));
+
+				// Register output resource used by the current capture method (texture or buffer)
+				typename TFrameType::FOutputResourceType OutputResource = CapturingFrame->RegisterResource(GraphBuilder);
+				{
+					TRACE_CPUPROFILER_EVENT_SCOPE(UMediaCapture::GraphSetup);
+					SCOPED_DRAW_EVENTF(Args.RHICmdList, MediaCapture, TEXT("MediaCapture"));
+
+					// If custom conversion was requested from implementation, give it useful information to apply 
+					if (Args.MediaCapture->ConversionOperation == EMediaCaptureConversionOperation::CUSTOM)
+					{
+						TRACE_CPUPROFILER_EVENT_SCOPE(UMediaCapture::CustomCapture);
+
+						Args.MediaCapture->OnCustomCapture_RenderingThread(GraphBuilder, CapturingFrame->CaptureBaseData, CapturingFrame->UserData
+							, SourceRGBTexture, OutputResource, ResolveParams, SizeU, SizeV);
+					}
+					else
+					{
+						TRACE_CPUPROFILER_EVENT_SCOPE(UMediaCapture::FormatConversion);
+						AddConversionPass(Args, { GraphBuilder, SourceRGBTexture, bRequiresFormatConversion, ResolveParams, SizeU, SizeV }, OutputResource);
+					}
+				}
+
+				// If Capture implementation is not grabbing GPU resource directly, push a readback pass to access it from CPU
+				if (Args.MediaCapture->bShouldCaptureRHIResource == false)
+				{
+					TRACE_CPUPROFILER_EVENT_SCOPE(UMediaCapture::EnqueueReadback);
+
+					CapturingFrame->EnqueueCopy(GraphBuilder, OutputResource);
+
+					CapturingFrame->bReadbackRequested = true;
+					const int32 FrameIndex = Args.MediaCapture->CaptureFrames.IndexOfByPredicate([CapturingFrame](const TUniquePtr<FCaptureFrame>& InCaptureFrame) { return InCaptureFrame.Get() == CapturingFrame; });
+					UE_LOG(LogMediaIOCore, VeryVerbose, TEXT("Requested copy for capture frame %d"), FrameIndex);
+					++Args.MediaCapture->NumberOfResourcesToReadback;
+				}
+
+				GraphBuilder.Execute();
+
+				if (Args.MediaCapture->bShouldCaptureRHIResource)
+				{
+					TRACE_CPUPROFILER_EVENT_SCOPE(UMediaCapture::RHIResourceCaptured);
+					SCOPE_CYCLE_COUNTER(STAT_MediaCapture_RenderThread_CaptureCallback)
+					Args.MediaCapture->OnRHIResourceCaptured_RenderingThread(CapturingFrame->CaptureBaseData, CapturingFrame->UserData, CapturingFrame->GetRHIResource());
+					CapturingFrame->bReadbackRequested = false;
+				}
+			}
+
+			return true;
+		}
+	};
+}
+
+
+
+/* namespace MediaCaptureDetails definition
+*****************************************************************************/
+
 namespace MediaCaptureDetails
 {
 	bool FindSceneViewportAndLevel(TSharedPtr<FSceneViewport>& OutSceneViewport);
@@ -131,14 +647,6 @@ UMediaCapture::FCaptureBaseData::FCaptureBaseData()
 
 }
 
-/* UMediaCapture::FCaptureFrame
-*****************************************************************************/
-UMediaCapture::FCaptureFrame::FCaptureFrame()
-	: bResolvedTargetRequested(false)
-{
-
-}
-
 /* FMediaCaptureOptions
 *****************************************************************************/
 FMediaCaptureOptions::FMediaCaptureOptions()
@@ -160,25 +668,21 @@ FMediaCaptureOptions::FMediaCaptureOptions()
 
 UMediaCapture::UMediaCapture(const FObjectInitializer& ObjectInitializer)
 	: Super(ObjectInitializer)
-	, CurrentResolvedTargetIndex(0)
-	, NumberOfCaptureFrame(2)
-	, CaptureRequestCount(0)
-	, MediaState(EMediaCaptureState::Stopped)
-	, DesiredSize(1280, 720)
-	, DesiredPixelFormat(EPixelFormat::PF_A2B10G10R10)
-	, DesiredOutputSize(1280, 720)
-	, DesiredOutputPixelFormat(EPixelFormat::PF_A2B10G10R10)
-	, ConversionOperation(EMediaCaptureConversionOperation::NONE)
-	, MediaOutputName(TEXT("[undefined]"))
-	, bUseRequestedTargetSize(false)
-	, bViewportHasFixedViewportSize(false)
-	, bResolvedTargetInitialized(false)
-	, bShouldCaptureRHITexture(false)
-	, WaitingForResolveCommandExecutionCounter(0)
-	, NumberOfTexturesToResolve(0)
+	, bOutputResourcesInitialized(false)
+	, bShouldCaptureRHIResource(false)
+	, WaitingForRenderCommandExecutionCounter(0)
+	, NumberOfResourcesToReadback(0)
+{
+}
+
+UMediaCapture::~UMediaCapture() = default;
+
+UMediaCapture::UMediaCapture(FVTableHelper&)
 {
+
 }
 
+
 void UMediaCapture::BeginDestroy()
 {
 	if (GetState() == EMediaCaptureState::Capturing || GetState() == EMediaCaptureState::Preparing)
@@ -254,11 +758,18 @@ bool UMediaCapture::CaptureSceneViewport(TSharedPtr<FSceneViewport>& InSceneView
 	bool bInitialized = CaptureSceneViewportImpl(InSceneViewport);
 
 	// This could have been updated by the call to CaptureSceneViewportImpl
-	bShouldCaptureRHITexture = ShouldCaptureRHITexture(); 
+	bShouldCaptureRHIResource = ShouldCaptureRHIResource(); 
+	if (bShouldCaptureRHIResource && DesiredOutputResourceType != EMediaCaptureResourceType::Texture)
+	{
+		UE_LOG(LogMediaIOCore, Warning, TEXT("Can't capture RHI resource when not using texture output resource"));
+		ResetFixedViewportSize(InSceneViewport, false);
+		SetState(EMediaCaptureState::Stopped);
+		MediaCaptureDetails::ShowSlateNotification();
+	}
 
 	if (bInitialized)
 	{
-		InitializeResolveTarget(MediaOutput->NumberOfTextureBuffers);
+		InitializeOutputResources(MediaOutput->NumberOfTextureBuffers);
 		bInitialized = GetState() != EMediaCaptureState::Stopped;
 	}
 
@@ -267,7 +778,7 @@ bool UMediaCapture::CaptureSceneViewport(TSharedPtr<FSceneViewport>& InSceneView
 		//no lock required, the command on the render thread is not active
 		CapturingSceneViewport = InSceneViewport;
 
-		CurrentResolvedTargetIndex = 0;
+		CurrentReadbackIndex = 0;
 		FCoreDelegates::OnEndFrame.AddUObject(this, &UMediaCapture::OnEndFrame_GameThread);
 	}
 	else
@@ -322,11 +833,11 @@ bool UMediaCapture::CaptureTextureRenderTarget2D(UTextureRenderTarget2D* InRende
 	bool bInitialized = CaptureRenderTargetImpl(InRenderTarget2D);
 
 	// This could have been updated by the call to CaptureSceneViewportImpl
-	bShouldCaptureRHITexture = ShouldCaptureRHITexture();
+	bShouldCaptureRHIResource = ShouldCaptureRHIResource();
 
 	if (bInitialized)
 	{
-		InitializeResolveTarget(MediaOutput->NumberOfTextureBuffers);
+		InitializeOutputResources(MediaOutput->NumberOfTextureBuffers);
 		bInitialized = GetState() != EMediaCaptureState::Stopped;
 	}
 
@@ -335,7 +846,7 @@ bool UMediaCapture::CaptureTextureRenderTarget2D(UTextureRenderTarget2D* InRende
 		//no lock required the command on the render thread is not active yet
 		CapturingRenderTarget = InRenderTarget2D;
 
-		CurrentResolvedTargetIndex = 0;
+		CurrentReadbackIndex = 0;
 		FCoreDelegates::OnEndFrame.AddUObject(this, &UMediaCapture::OnEndFrame_GameThread);
 	}
 	else
@@ -363,9 +874,11 @@ void UMediaCapture::CacheMediaOutput(EMediaCaptureSourceType InSourceType)
 void UMediaCapture::CacheOutputOptions()
 {
 	DesiredOutputSize = GetOutputSize(DesiredSize, ConversionOperation);
+	DesiredOutputResourceType = GetOutputResourceType(ConversionOperation);
 	DesiredOutputPixelFormat = GetOutputPixelFormat(DesiredPixelFormat, ConversionOperation);
+	DesiredOutputBufferDescription = GetOutputBufferDescription(ConversionOperation);
 	MediaOutputName = *MediaOutput->GetName();
-	bShouldCaptureRHITexture = ShouldCaptureRHITexture();
+	bShouldCaptureRHIResource = ShouldCaptureRHIResource();
 }
 
 FIntPoint UMediaCapture::GetOutputSize(const FIntPoint & InSize, const EMediaCaptureConversionOperation & InConversionOperation) const
@@ -401,6 +914,28 @@ EPixelFormat UMediaCapture::GetOutputPixelFormat(const EPixelFormat & InPixelFor
 	}
 }
 
+EMediaCaptureResourceType UMediaCapture::GetOutputResourceType(const EMediaCaptureConversionOperation& InConversionOperation) const
+{
+	switch (InConversionOperation)
+	{
+	case EMediaCaptureConversionOperation::CUSTOM:
+		return GetCustomOutputResourceType();
+	default:
+		return EMediaCaptureResourceType::Texture;
+	}
+}
+
+FRDGBufferDesc UMediaCapture::GetOutputBufferDescription(EMediaCaptureConversionOperation InConversionOperation) const
+{
+	switch (InConversionOperation)
+	{
+	case EMediaCaptureConversionOperation::CUSTOM:
+		return GetCustomBufferDescription(DesiredSize);
+	default:
+		return FRDGBufferDesc();
+	}
+}
+
 bool UMediaCapture::UpdateSceneViewport(TSharedPtr<FSceneViewport>& InSceneViewport)
 {
 	if (!MediaCaptureDetails::ValidateIsCapturing(*this))
@@ -435,7 +970,7 @@ bool UMediaCapture::UpdateSceneViewport(TSharedPtr<FSceneViewport>& InSceneViewp
 
 	{
 		FScopeLock Lock(&AccessingCapturingSource);
-		while(WaitingForResolveCommandExecutionCounter.Load() > 0)
+		while(WaitingForRenderCommandExecutionCounter.Load() > 0)
 		{
 			FlushRenderingCommands();
 		}
@@ -480,7 +1015,7 @@ bool UMediaCapture::UpdateTextureRenderTarget2D(UTextureRenderTarget2D * InRende
 
 	{
 		FScopeLock Lock(&AccessingCapturingSource);
-		while (WaitingForResolveCommandExecutionCounter.Load() > 0)
+		while (WaitingForRenderCommandExecutionCounter.Load() > 0)
 		{
 			FlushRenderingCommands();
 		}
@@ -510,7 +1045,7 @@ void UMediaCapture::StopCapture(bool bAllowPendingFrameToBeProcess)
 			//Do not flush when auto stopping to avoid hitches.
 			if(DesiredCaptureOptions.bAutostopOnCapture != true)
 			{
-			while (WaitingForResolveCommandExecutionCounter.Load() > 0)
+			while (WaitingForRenderCommandExecutionCounter.Load() > 0)
 			{
 				FlushRenderingCommands();
 			}
@@ -525,7 +1060,7 @@ void UMediaCapture::StopCapture(bool bAllowPendingFrameToBeProcess)
 
 			FCoreDelegates::OnEndFrame.RemoveAll(this);
 
-			while (WaitingForResolveCommandExecutionCounter.Load() > 0 || !bResolvedTargetInitialized)
+			while (WaitingForRenderCommandExecutionCounter.Load() > 0 || !bOutputResourcesInitialized)
 			{
 				FlushRenderingCommands();
 			}
@@ -601,7 +1136,7 @@ void UMediaCapture::ResetFixedViewportSize(TSharedPtr<FSceneViewport> InViewport
 {
 	if (bViewportHasFixedViewportSize && InViewport.IsValid())
 	{
-		if (bInFlushRenderingCommands && WaitingForResolveCommandExecutionCounter.Load() > 0)
+		if (bInFlushRenderingCommands && WaitingForRenderCommandExecutionCounter.Load() > 0)
 		{
 			FlushRenderingCommands();
 		}
@@ -612,13 +1147,15 @@ void UMediaCapture::ResetFixedViewportSize(TSharedPtr<FSceneViewport> InViewport
 
 bool UMediaCapture::HasFinishedProcessing() const
 {
-	return WaitingForResolveCommandExecutionCounter.Load() == 0
+	return WaitingForRenderCommandExecutionCounter.Load() == 0
 		|| GetState() == EMediaCaptureState::Error
 		|| GetState() == EMediaCaptureState::Stopped;
 }
 
-void UMediaCapture::InitializeResolveTarget(int32 InNumberOfBuffers)
+void UMediaCapture::InitializeOutputResources(int32 InNumberOfBuffers)
 {
+	using namespace UE::MediaCaptureData;
+
 	if (DesiredOutputSize.X <= 0 || DesiredOutputSize.Y <= 0)
 	{
 		UE_LOG(LogMediaIOCore, Error, TEXT("Can't start the capture. The size requested is negative or zero."));
@@ -628,42 +1165,59 @@ void UMediaCapture::InitializeResolveTarget(int32 InNumberOfBuffers)
 
 	NumberOfCaptureFrame = InNumberOfBuffers;
 	check(CaptureFrames.Num() == 0);
-	CaptureFrames.AddDefaulted(InNumberOfBuffers);
 
-		UMediaCapture* This = this;
-		ENQUEUE_RENDER_COMMAND(MediaOutputCaptureFrameCreateTexture)(
-			[This](FRHICommandListImmediate& RHICmdList)
+	UMediaCapture* This = this;
+	ENQUEUE_RENDER_COMMAND(MediaOutputCaptureFrameCreateResources)(
+		[This](FRHICommandListImmediate& RHICmdList)
+		{
+			FRHIResourceCreateInfo CreateInfo(TEXT("UMediaCapture"));
+			for (int32 Index = 0; Index < This->NumberOfCaptureFrame; ++Index)
 			{
-				FRHIResourceCreateInfo CreateInfo(TEXT("UMediaCapture"));
-				for (int32 Index = 0; Index < This->NumberOfCaptureFrame; ++Index)
-				{
-				FPooledRenderTargetDesc OutputDesc = FPooledRenderTargetDesc::Create2DDesc(
-					This->DesiredOutputSize,
-					This->DesiredOutputPixelFormat,
-					FClearValueBinding::None,
-					TexCreate_Shared,
-					TexCreate_RenderTargetable,
-					false);
-
-				GRenderTargetPool.FindFreeElement(RHICmdList, OutputDesc, This->CaptureFrames[Index].RenderTarget, *FString::Format(TEXT("MediaCapture RenderTarget {0}"), { Index }));
-
-				// Only create CPU readback texture when we are using the CPU callback
-				if (!This->bShouldCaptureRHITexture)
+				if (This->DesiredOutputResourceType == EMediaCaptureResourceType::Texture)
 				{
-					This->CaptureFrames[Index].ReadbackTexture = RHICreateTexture2D(
-						This->DesiredOutputSize.X,
-						This->DesiredOutputSize.Y,
+					TUniquePtr<FTextureCaptureFrame> NewFrame = MakeUnique<FTextureCaptureFrame>();
+					FRDGTextureDesc OutputDesc = FRDGTextureDesc::Create2D(
+						This->DesiredOutputSize,
 						This->DesiredOutputPixelFormat,
-						1,
-						1,
-						TexCreate_CPUReadback,
-						CreateInfo
-					);
+						FClearValueBinding::None,
+						TexCreate_Shared | TexCreate_RenderTargetable | TexCreate_UAV);
+
+
+					NewFrame->RenderTarget = AllocatePooledTexture(OutputDesc, *FString::Format(TEXT("MediaCapture RenderTarget {0}"), { Index }));
+
+					// Only create CPU readback resource when we are using the CPU callback
+					if (!This->bShouldCaptureRHIResource)
+					{
+						NewFrame->ReadbackTexture = MakeUnique<FRHIGPUTextureReadback>(*FString::Printf(TEXT("MediaCaptureTextureReadback_%d"), Index));
+					}
+
+					This->CaptureFrames.Emplace(MoveTemp(NewFrame));
 				}
+				else
+				{
+					TUniquePtr<FBufferCaptureFrame> NewFrame = MakeUnique<FBufferCaptureFrame>();
+
+					if (This->DesiredOutputBufferDescription.NumElements > 0)
+					{
+						NewFrame->Buffer = AllocatePooledBuffer(This->DesiredOutputBufferDescription, *FString::Format(TEXT("MediaCapture BufferResource {0}"), { Index }));
 
+						// Only create CPU readback resource when we are using the CPU callback
+						if (!This->bShouldCaptureRHIResource)
+						{
+							NewFrame->ReadbackBuffer = MakeUnique<FRHIGPUBufferReadback>(*FString::Printf(TEXT("MediaCaptureBufferReadback_%d"), Index));
+						}
+
+						This->CaptureFrames.Emplace(MoveTemp(NewFrame));
+					}
+					else
+					{
+						UE_LOG(LogMediaIOCore, Error, TEXT("Can't start the capture. Trying to allocate buffer resource but number of elements to allocate was 0."));
+						This->SetState(EMediaCaptureState::Error);
+					}
+				}
 			}
-				This->bResolvedTargetInitialized = true;
-			});
+			This->bOutputResourcesInitialized = true;
+		});
 }
 
 bool UMediaCapture::ValidateMediaOutput() const
@@ -692,10 +1246,12 @@ bool UMediaCapture::ValidateMediaOutput() const
 
 void UMediaCapture::OnEndFrame_GameThread()
 {
-	const FString EndFrameTraceName = FString::Format(TEXT("MediaCapture::OnEndFrame_GameThread (Index {0})"), { CurrentResolvedTargetIndex });
+	using namespace UE::MediaCaptureData;
+
+	const FString EndFrameTraceName = FString::Format(TEXT("MediaCapture::OnEndFrame_GameThread (Index {0})"), { CurrentReadbackIndex });
 	TRACE_CPUPROFILER_EVENT_SCOPE_TEXT(*EndFrameTraceName);
 
-	if (!bResolvedTargetInitialized)
+	if (!bOutputResourcesInitialized)
 	{
 		FlushRenderingCommands();
 	}
@@ -720,34 +1276,33 @@ void UMediaCapture::OnEndFrame_GameThread()
 		return;
 	}
 
-	CurrentResolvedTargetIndex = (CurrentResolvedTargetIndex + 1) % NumberOfCaptureFrame;
-	int32 ReadyFrameIndex = (CurrentResolvedTargetIndex - 1); // Previous one in the buffer queue
+	CurrentReadbackIndex = (CurrentReadbackIndex + 1) % NumberOfCaptureFrame;
+	int32 ReadyFrameIndex = (CurrentReadbackIndex - 1); // Previous one in the buffer queue
 	if (ReadyFrameIndex < 0)
 	{
 		ReadyFrameIndex = NumberOfCaptureFrame - 1;
 	}
 
 	// Frame that should be on the system ram and we want to send to the user
-	FCaptureFrame* ReadyFrame = &CaptureFrames[ReadyFrameIndex];
+	FCaptureFrame* ReadyFrame = CaptureFrames[ReadyFrameIndex].Get();
 	// Frame that we want to transfer to system ram
-	FCaptureFrame* CapturingFrame = (GetState() != EMediaCaptureState::StopRequested) ? &CaptureFrames[CurrentResolvedTargetIndex] : nullptr;
+	FCaptureFrame* CapturingFrame = (GetState() != EMediaCaptureState::StopRequested) ? CaptureFrames[CurrentReadbackIndex].Get() : nullptr;
 
-	if (!ShouldCaptureRHITexture())
+	if (!ShouldCaptureRHIResource())
 	{
-		TStringBuilder<256> ResolveTargetInfoBuilder;
-		ResolveTargetInfoBuilder << "\n";
+		TStringBuilder<256> ReadbackInfoBuilder;
+		ReadbackInfoBuilder << "\n";
 		for (int32 Index = 0; Index < NumberOfCaptureFrame; Index++)
 		{
-			ResolveTargetInfoBuilder << FString::Format(TEXT("Frame {0} requested: {1}\n"), { Index, CaptureFrames[Index].bResolvedTargetRequested ? TEXT("true") : TEXT("false") });
+			ReadbackInfoBuilder << FString::Format(TEXT("Frame {0} readback requested: {1}\n"), { Index, CaptureFrames[Index]->bReadbackRequested ? TEXT("true") : TEXT("false") });
 		}
-		FString FinalString = ResolveTargetInfoBuilder.ToString();
-		UE_LOG(LogMediaIOCore, VeryVerbose, TEXT("%s"), *FinalString);
+		UE_LOG(LogMediaIOCore, VeryVerbose, TEXT("%s"), ReadbackInfoBuilder.GetData());
 	}
 
 
-	UE_LOG(LogMediaIOCore, VeryVerbose, TEXT("MediaOutput: '%s'. ReadyFrameIndex: '%d' '%s'. CurrentResolvedTargetIndex: '%d'.")
-		, *MediaOutputName, ReadyFrameIndex, (CaptureFrames[ReadyFrameIndex].bResolvedTargetRequested) ? TEXT("Y"): TEXT("N"), CurrentResolvedTargetIndex);
-	if (GetState() == EMediaCaptureState::StopRequested && NumberOfTexturesToResolve.Load() <= 0)
+	UE_LOG(LogMediaIOCore, VeryVerbose, TEXT("MediaOutput: '%s'. ReadyFrameIndex: '%d' '%s'. CurrentReadbackIndex: '%d'.")
+		, *MediaOutputName, ReadyFrameIndex, (CaptureFrames[ReadyFrameIndex]->bReadbackRequested) ? TEXT("Y"): TEXT("N"), CurrentReadbackIndex);
+	if (GetState() == EMediaCaptureState::StopRequested && NumberOfResourcesToReadback.Load() <= 0)
 	{
 		// All the requested frames have been captured.
 		StopCapture(false);
@@ -757,9 +1312,9 @@ void UMediaCapture::OnEndFrame_GameThread()
 	if (CapturingFrame)
 	{
 		//Verify if game thread is overrunning the render thread.
-		if (CapturingFrame->bResolvedTargetRequested)
+		if (CapturingFrame->bReadbackRequested)
 		{
-			FString TraceName = FString::Format(TEXT("MediaCapture::FlushRenderingCommands (Index {0}/{1})"), { CurrentResolvedTargetIndex, NumberOfCaptureFrame });
+			FString TraceName = FString::Format(TEXT("MediaCapture::FlushRenderingCommands (Index {0}/{1})"), { CurrentReadbackIndex, NumberOfCaptureFrame });
 			TRACE_CPUPROFILER_EVENT_SCOPE_TEXT(*TraceName);
 			UE_LOG(LogMediaIOCore, VeryVerbose, TEXT("MediaOutput: '%s'. Flushing commands.")
 				, *MediaOutputName);
@@ -787,7 +1342,7 @@ void UMediaCapture::OnEndFrame_GameThread()
 
 		if (InCapturingSceneViewport != nullptr || InTextureRenderTargetResource != nullptr)
 		{
-			++WaitingForResolveCommandExecutionCounter;
+			++WaitingForRenderCommandExecutionCounter;
 
 			// RenderCommand to be executed on the RenderThread
 			ENQUEUE_RENDER_COMMAND(FMediaOutputCaptureFrameCreateTexture)(
@@ -807,267 +1362,68 @@ void UMediaCapture::OnEndFrame_GameThread()
 
 void UMediaCapture::Capture_RenderThread(FRHICommandListImmediate& RHICmdList,
 	UMediaCapture* InMediaCapture,
-	FCaptureFrame* CapturingFrame,
-	FCaptureFrame* ReadyFrame,
+	UE::MediaCaptureData::FCaptureFrame* CapturingFrame,
+	UE::MediaCaptureData::FCaptureFrame* ReadyFrame,
 	FSceneViewport* InCapturingSceneViewport,
 	FTextureRenderTargetResource* InTextureRenderTargetResource,
 	FIntPoint InDesiredSize,
 	FMediaCaptureStateChangedSignature InOnStateChanged)
 {
+	using namespace UE::MediaCaptureData;
+	
 	TRACE_CPUPROFILER_EVENT_SCOPE(UMediaCapture::Capture_RenderThread);
-	FTexture2DRHIRef SourceTexture;
 
-	if (CapturingFrame)
+	// Whatever happens, we want to decrement our counter to track enqueued commands
+	ON_SCOPE_EXIT
 	{
-		if (!CapturingFrame->RenderTarget)
-		{
-			InMediaCapture->SetState(EMediaCaptureState::Error);
-			UE_LOG(LogMediaIOCore, Error, TEXT("The capture will stop for '%s'. A capture frame had an invalid render target."), *InMediaCapture->MediaOutputName);
-		}
-		else
-		{
-			TRACE_CPUPROFILER_EVENT_SCOPE(UMediaCapture::BeforeFrameCaptured);
-			UE_LOG(LogMediaIOCore, VeryVerbose, TEXT("Before Locking texture %p"), reinterpret_cast<uintptr_t>(CapturingFrame->RenderTarget->GetRHI()->GetNativeResource()));
-			BeforeFrameCaptured_RenderingThread(CapturingFrame->CaptureBaseData, CapturingFrame->UserData, CapturingFrame->RenderTarget->GetRHI());
-			UE_LOG(LogMediaIOCore, VeryVerbose, TEXT("After Locking texture %p"), reinterpret_cast<uintptr_t>(CapturingFrame->RenderTarget->GetRHI()->GetNativeResource()));
-		}
-	}
+		--InMediaCapture->WaitingForRenderCommandExecutionCounter;
+	};
 
+	FTexture2DRHIRef SourceTexture;
+
+	if (CapturingFrame)
 	{
-		if (InCapturingSceneViewport)
+		// Call the capture frame algo based on the specific type of resource we are using
+		bool bHasCaptureSuceeded = false;
+		FCaptureFrameArgs CaptureArgs = { RHICmdList, InMediaCapture, InCapturingSceneViewport, InTextureRenderTargetResource, InDesiredSize };
+		if (DesiredOutputResourceType == EMediaCaptureResourceType::Texture)
 		{
-#if WITH_EDITOR
-			if (!IsRunningGame())
+			if (ensure(CapturingFrame->IsTextureResource()))
 			{
-				// PIE, PIE in windows, editor viewport
-				SourceTexture = InCapturingSceneViewport->GetRenderTargetTexture();
-				if (!SourceTexture.IsValid() && InCapturingSceneViewport->GetViewportRHI())
-				{
-					SourceTexture = RHICmdList.GetViewportBackBuffer(InCapturingSceneViewport->GetViewportRHI());
-				}
+				FTextureCaptureFrame* TextureFrame = static_cast<FTextureCaptureFrame*>(CapturingFrame);
+				bHasCaptureSuceeded = FMediaCaptureHelper::CaptureFrame(CaptureArgs, TextureFrame);
 			}
 			else
-#endif
-				if (InCapturingSceneViewport->GetViewportRHI())
-				{
-					// Standalone and packaged
-					SourceTexture = RHICmdList.GetViewportBackBuffer(InCapturingSceneViewport->GetViewportRHI());
-				}
-		}
-		else if (InTextureRenderTargetResource && InTextureRenderTargetResource->GetTextureRenderTarget2DResource())
-		{
-			SourceTexture = InTextureRenderTargetResource->GetTextureRenderTarget2DResource()->GetTextureRHI();
-		}
-	}
-
-	if (!SourceTexture.IsValid())
-	{
-		InMediaCapture->SetState(EMediaCaptureState::Error);
-		UE_LOG(LogMediaIOCore, Error, TEXT("Can't grab the Texture to capture for '%s'."), *InMediaCapture->MediaOutputName);
-	}
-	else if (CapturingFrame)
-	{
-		// If it is a simple rgba swizzle we can handle the conversion. Supported formats
-		// contained in SupportedRgbaSwizzleFormats. Warning would've been displayed on start of capture.
-		if (InMediaCapture->DesiredPixelFormat != SourceTexture->GetFormat() && 
-			(!SupportedRgbaSwizzleFormats.Contains(SourceTexture->GetFormat()) || !InMediaCapture->DesiredCaptureOptions.bConvertToDesiredPixelFormat))
-		{
-			InMediaCapture->SetState(EMediaCaptureState::Error);
-			UE_LOG(LogMediaIOCore, Error, TEXT("The capture will stop for '%s'. The Source pixel format doesn't match with the user requested pixel format. %sRequested: %s Source: %s")
-				, *InMediaCapture->MediaOutputName
-				, (SupportedRgbaSwizzleFormats.Contains(SourceTexture->GetFormat()) && !InMediaCapture->DesiredCaptureOptions.bConvertToDesiredPixelFormat) ? TEXT("Please enable \"Convert To Desired Pixel Format\" option in Media Capture settings. ") : TEXT("")
-				, GetPixelFormatString(InMediaCapture->DesiredPixelFormat)
-				, GetPixelFormatString(SourceTexture->GetFormat()));
-		}
-
-		if (InMediaCapture->DesiredCaptureOptions.Crop == EMediaCaptureCroppingType::None)
-		{
-			if (InDesiredSize.X != SourceTexture->GetSizeX() || InDesiredSize.Y != SourceTexture->GetSizeY())
 			{
-				InMediaCapture->SetState(EMediaCaptureState::Error);
-				UE_LOG(LogMediaIOCore, Error, TEXT("The capture will stop for '%s'. The Source size doesn't match with the user requested size. Requested: %d,%d  Source: %d,%d")
-					, *InMediaCapture->MediaOutputName
-					, InDesiredSize.X, InDesiredSize.Y
-					, SourceTexture->GetSizeX(), SourceTexture->GetSizeY());
+				UE_LOG(LogMediaIOCore, Error, TEXT("The capture will stop for '%s'. Capture frame was expected to use Texture resource but wasn't."), *InMediaCapture->MediaOutputName);
 			}
 		}
 		else
 		{
-			FIntPoint StartCapturePoint = FIntPoint::ZeroValue;
-			if (InMediaCapture->DesiredCaptureOptions.Crop == EMediaCaptureCroppingType::Custom)
+			if (ensure(CapturingFrame->IsBufferResource()))
 			{
-				StartCapturePoint = InMediaCapture->DesiredCaptureOptions.CustomCapturePoint;
+				FBufferCaptureFrame* BufferFrame = static_cast<FBufferCaptureFrame*>(CapturingFrame);
+				bHasCaptureSuceeded = FMediaCaptureHelper::CaptureFrame(CaptureArgs, BufferFrame);
 			}
-
-			if ((uint32)(InDesiredSize.X + StartCapturePoint.X) > SourceTexture->GetSizeX() || (uint32)(InDesiredSize.Y + StartCapturePoint.Y) > SourceTexture->GetSizeY())
-			{
-				InMediaCapture->SetState(EMediaCaptureState::Error);
-				UE_LOG(LogMediaIOCore, Error, TEXT("The capture will stop for '%s'. The Source size doesn't match with the user requested size. Requested: %d,%d  Source: %d,%d")
-					, *InMediaCapture->MediaOutputName
-					, InDesiredSize.X, InDesiredSize.Y
-					, SourceTexture->GetSizeX(), SourceTexture->GetSizeY());
-			}
-		}
-	}
-
-	if (CapturingFrame && InMediaCapture->GetState() != EMediaCaptureState::Error)
-	{
-		TRACE_CPUPROFILER_EVENT_SCOPE(UMediaCapture::CopyToResolve);
-		SCOPE_CYCLE_COUNTER(STAT_MediaCapture_RenderThread_CopyToResolve);
-
-		FRHITexture* DestRenderTarget = CapturingFrame->RenderTarget->GetRHI();
-		// Do we need to crop
-		float ULeft = 0.0f;
-		float URight = 1.0f;
-		float VTop = 0.0f;
-		float VBottom = 1.0f;
-		FResolveParams ResolveParams;
-		if (InMediaCapture->DesiredCaptureOptions.Crop != EMediaCaptureCroppingType::None)
-		{
-			switch (InMediaCapture->DesiredCaptureOptions.Crop)
+			else
 			{
-			case EMediaCaptureCroppingType::Center:
-				ResolveParams.Rect = FResolveRect((SourceTexture->GetSizeX() - InDesiredSize.X) / 2, (SourceTexture->GetSizeY() - InDesiredSize.Y) / 2, 0, 0);
-				ResolveParams.Rect.X2 = ResolveParams.Rect.X1 + InDesiredSize.X;
-				ResolveParams.Rect.Y2 = ResolveParams.Rect.Y1 + InDesiredSize.Y;
-				break;
-			case EMediaCaptureCroppingType::TopLeft:
-				ResolveParams.Rect = FResolveRect(0, 0, InDesiredSize.X, InDesiredSize.Y);
-				break;
-			case EMediaCaptureCroppingType::Custom:
-				ResolveParams.Rect = FResolveRect(InMediaCapture->DesiredCaptureOptions.CustomCapturePoint.X, InMediaCapture->DesiredCaptureOptions.CustomCapturePoint.Y, 0, 0);
-				ResolveParams.Rect.X2 = ResolveParams.Rect.X1 + InDesiredSize.X;
-				ResolveParams.Rect.Y2 = ResolveParams.Rect.Y1 + InDesiredSize.Y;
-				break;
+				UE_LOG(LogMediaIOCore, Error, TEXT("The capture will stop for '%s'. Capture frame was expected to use Buffer resource but wasn't."), *InMediaCapture->MediaOutputName);
 			}
-
-			ResolveParams.DestRect.X1 = 0;
-			ResolveParams.DestRect.X2 = InDesiredSize.X;
-			ResolveParams.DestRect.Y1 = 0;
-			ResolveParams.DestRect.Y2 = InDesiredSize.Y;
-
-			ULeft = (float)ResolveParams.Rect.X1 / (float)SourceTexture->GetSizeX();
-			URight = (float)ResolveParams.Rect.X2 / (float)SourceTexture->GetSizeX();
-			VTop = (float)ResolveParams.Rect.Y1 / (float)SourceTexture->GetSizeY();
-			VBottom = (float)ResolveParams.Rect.Y2 / (float)SourceTexture->GetSizeY();
 		}
-
+		
+		if (bHasCaptureSuceeded == false)
 		{
-			TRACE_CPUPROFILER_EVENT_SCOPE(UMediaCapture::DrawViewportInRenderTarget);
-			SCOPED_DRAW_EVENTF(RHICmdList, MediaCapture, TEXT("MediaCapture"));
-
-			bool bRequiresFormatConversion = InMediaCapture->DesiredPixelFormat != SourceTexture->GetFormat();
-
-			if (InMediaCapture->ConversionOperation == EMediaCaptureConversionOperation::NONE && !bRequiresFormatConversion)
-			{
-				// Asynchronously copy target from GPU to GPU
-				RHICmdList.CopyToResolveTarget(SourceTexture, DestRenderTarget, ResolveParams);
-			}
-			else if (InMediaCapture->ConversionOperation == EMediaCaptureConversionOperation::CUSTOM)
-			{
-				InMediaCapture->OnCustomCapture_RenderingThread(RHICmdList, CapturingFrame->CaptureBaseData, CapturingFrame->UserData
-					, SourceTexture, DestRenderTarget, ResolveParams, {ULeft, URight}, {VTop, VBottom});
-			}
-			else
-			{
-				TRACE_CPUPROFILER_EVENT_SCOPE(UMediaCapture::FormatConversion);
-				// convert the source with a draw call
-				FGraphicsPipelineStateInitializer GraphicsPSOInit;
-				FRHITexture* RenderTarget = DestRenderTarget;
-				RHICmdList.Transition(FRHITransitionInfo(RenderTarget, ERHIAccess::Unknown, ERHIAccess::RTV));
-				FRHIRenderPassInfo RPInfo(RenderTarget,  ERenderTargetActions::DontLoad_Store);
-				RHICmdList.BeginRenderPass(RPInfo, TEXT("MediaCapture"));
-
-				RHICmdList.ApplyCachedRenderTargets(GraphicsPSOInit);
-
-				GraphicsPSOInit.DepthStencilState = TStaticDepthStencilState<false, CF_Always>::GetRHI();
-				GraphicsPSOInit.RasterizerState = TStaticRasterizerState<>::GetRHI();
-				GraphicsPSOInit.BlendState = TStaticBlendStateWriteMask<CW_RGBA, CW_NONE, CW_NONE, CW_NONE, CW_NONE, CW_NONE, CW_NONE, CW_NONE>::GetRHI();
-				GraphicsPSOInit.PrimitiveType = PT_TriangleStrip;
-
-				// configure media shaders
-				auto ShaderMap = GetGlobalShaderMap(GMaxRHIFeatureLevel);
-				TShaderMapRef<FMediaShadersVS> VertexShader(ShaderMap);
-
-				GraphicsPSOInit.BoundShaderState.VertexDeclarationRHI = GMediaVertexDeclaration.VertexDeclarationRHI;
-				GraphicsPSOInit.BoundShaderState.VertexShaderRHI = VertexShader.GetVertexShader();
-
-				const bool bDoLinearToSRGB = false;
-
-				switch (InMediaCapture->ConversionOperation)
-				{
-				case EMediaCaptureConversionOperation::RGBA8_TO_YUV_8BIT:
-				{
-					TShaderMapRef<FRGB8toUYVY8ConvertPS> ConvertShader(ShaderMap);
-					GraphicsPSOInit.BoundShaderState.PixelShaderRHI = ConvertShader.GetPixelShader();
-					SetGraphicsPipelineState(RHICmdList, GraphicsPSOInit, 0);
-					ConvertShader->SetParameters(RHICmdList, SourceTexture, MediaShaders::RgbToYuvRec709Scaled, MediaShaders::YUVOffset8bits, bDoLinearToSRGB);
-				}
-				break;
-				case EMediaCaptureConversionOperation::RGB10_TO_YUVv210_10BIT:
-				{
-					TShaderMapRef<FRGB10toYUVv210ConvertPS> ConvertShader(ShaderMap);
-					GraphicsPSOInit.BoundShaderState.PixelShaderRHI = ConvertShader.GetPixelShader();
-					SetGraphicsPipelineState(RHICmdList, GraphicsPSOInit, 0);
-					ConvertShader->SetParameters(RHICmdList, SourceTexture, MediaShaders::RgbToYuvRec709Scaled, MediaShaders::YUVOffset10bits, bDoLinearToSRGB);
-				}
-				break;
-				case EMediaCaptureConversionOperation::INVERT_ALPHA:
-					// fall through
-				case EMediaCaptureConversionOperation::SET_ALPHA_ONE:
-					// fall through
-				case EMediaCaptureConversionOperation::NONE:
-					bRequiresFormatConversion = true;
-				default:
-					if (bRequiresFormatConversion)
-					{
-						FModifyAlphaSwizzleRgbaPS::FPermutationDomain PermutationVector;
-						// In cases where texture is converted from a format that doesn't have A channel, we want to force set it to 1.
-						EMediaCaptureConversionOperation MediaConversionOperation = InMediaCapture->DesiredCaptureOptions.bForceAlphaToOneOnConversion ? EMediaCaptureConversionOperation::SET_ALPHA_ONE : InMediaCapture->ConversionOperation;
-						PermutationVector.Set<FModifyAlphaSwizzleRgbaPS::FConversionOp>(static_cast<int32>(MediaConversionOperation));
-						TShaderMapRef<FModifyAlphaSwizzleRgbaPS> ConvertShader(ShaderMap, PermutationVector);
-						GraphicsPSOInit.BoundShaderState.PixelShaderRHI = ConvertShader.GetPixelShader();
-						SetGraphicsPipelineState(RHICmdList, GraphicsPSOInit, 0);
-						ConvertShader->SetParameters(RHICmdList, SourceTexture);
-					}
-				break;
-				}
-
-				// draw full size quad into render target
-				FBufferRHIRef VertexBuffer = CreateTempMediaVertexBuffer(ULeft, URight, VTop, VBottom);
-				RHICmdList.SetStreamSource(0, VertexBuffer, 0);
-
-				// set viewport to RT size
-				RHICmdList.SetViewport(0, 0, 0.0f, InMediaCapture->DesiredOutputSize.X, InMediaCapture->DesiredOutputSize.Y, 1.0f);
-				RHICmdList.DrawPrimitive(0, 2, 1);
-				RHICmdList.EndRenderPass();
-				RHICmdList.Transition(FRHITransitionInfo(DestRenderTarget, ERHIAccess::RTV, ERHIAccess::SRVGraphics));
-
-			}
-
-			if (InMediaCapture->bShouldCaptureRHITexture)
-			{
-				InMediaCapture->OnRHITextureCaptured_RenderingThread(CapturingFrame->CaptureBaseData, CapturingFrame->UserData, DestRenderTarget);
-				CapturingFrame->bResolvedTargetRequested = false;
-			}
-			else
-			{
-				// Asynchronously copy duplicate target from GPU to System Memory
-				RHICmdList.CopyToResolveTarget(DestRenderTarget, CapturingFrame->ReadbackTexture, FResolveParams());
-				CapturingFrame->bResolvedTargetRequested = true;
-				const int32 FrameIndex = CaptureFrames.IndexOfByPredicate([&CapturingFrame](const FCaptureFrame& InCaptureFrame) { return &InCaptureFrame == CapturingFrame; });
-				UE_LOG(LogMediaIOCore, VeryVerbose, TEXT("Requested copy for capture frame %d"), FrameIndex);
-				++NumberOfTexturesToResolve;
-			}
+			InMediaCapture->SetState(EMediaCaptureState::Error);
 		}
 	}
 
-	if (!InMediaCapture->bShouldCaptureRHITexture && InMediaCapture->GetState() != EMediaCaptureState::Error)
+	if (!InMediaCapture->bShouldCaptureRHIResource && InMediaCapture->GetState() != EMediaCaptureState::Error)
 	{
-		if (ReadyFrame->bResolvedTargetRequested)
+		if (ReadyFrame->bReadbackRequested)
 		{
+			FRHIGPUMask GPUMask;
 #if WITH_MGPU
-			FRHIGPUMask GPUMask = RHICmdList.GetGPUMask();
+			GPUMask = RHICmdList.GetGPUMask();
 
 			// If GPUMask is not set to a specific GPU we and since we are reading back the texture, it shouldn't matter which GPU we do this on.
 			if (!GPUMask.HasSingleIndex())
@@ -1077,34 +1433,35 @@ void UMediaCapture::Capture_RenderThread(FRHICommandListImmediate& RHICmdList,
 
 			SCOPED_GPU_MASK(RHICmdList, GPUMask);
 #endif
-			check(ReadyFrame->ReadbackTexture.IsValid());
-
 			// Lock & read
 			void* ColorDataBuffer = nullptr;
-			int32 Width = 0, Height = 0;
+			int32 Width = 0;
+			int32 Height = 0;
+			int32 RowPitchInPixels = 0;
+
+			// Texture readback does no verification if it's ready. It is assumed it will stall if it's not
+			// Buffer readback will verify and block GPU until idle in case it's not ready otherwise it ensures
 			{
-				SCOPE_CYCLE_COUNTER(STAT_MediaCapture_RenderThread_MapStaging);
-				RHICmdList.MapStagingSurface(ReadyFrame->ReadbackTexture, ColorDataBuffer, Width, Height);
+				SCOPE_CYCLE_COUNTER(STAT_MediaCapture_RenderThread_LockResource);
+				ColorDataBuffer = ReadyFrame->Lock(RHICmdList, GPUMask, RowPitchInPixels);
 			}
 
 			{
-				SCOPE_CYCLE_COUNTER(STAT_MediaCapture_RenderThread_Callback);
+				SCOPE_CYCLE_COUNTER(STAT_MediaCapture_RenderThread_CaptureCallback)
 
 				// The Width/Height of the surface may be different then the DesiredOutputSize : Some underlying implementations enforce a specific stride, therefore
 				// there may be padding at the end of each row.
 				InMediaCapture->OnFrameCaptured_RenderingThread(ReadyFrame->CaptureBaseData, ReadyFrame->UserData, ColorDataBuffer, InMediaCapture->DesiredOutputSize.X, InMediaCapture->DesiredOutputSize.Y, Width * 4);
 			}
-			ReadyFrame->bResolvedTargetRequested = false;
+			ReadyFrame->bReadbackRequested = false;
 
-			const int32 FrameIndex = CaptureFrames.IndexOfByPredicate([&CapturingFrame](const FCaptureFrame& InCaptureFrame) { return &InCaptureFrame == CapturingFrame; });
+			const int32 FrameIndex = CaptureFrames.IndexOfByPredicate([&CapturingFrame](const TUniquePtr<FCaptureFrame>& InCaptureFrame) { return InCaptureFrame.Get() == CapturingFrame; });
 			UE_LOG(LogMediaIOCore, VeryVerbose, TEXT("Copy completed for capture frame %d."), FrameIndex);
-			--NumberOfTexturesToResolve;
+			--NumberOfResourcesToReadback;
 
-			RHICmdList.UnmapStagingSurface(ReadyFrame->ReadbackTexture);
+			ReadyFrame->Unlock();
 		}
 	}
-
-	--InMediaCapture->WaitingForResolveCommandExecutionCounter;
 }
 
 /* namespace MediaCaptureDetails implementation
diff --git a/Engine/Plugins/Media/MediaIOFramework/Source/MediaIOCore/Public/MediaCapture.h b/Engine/Plugins/Media/MediaIOFramework/Source/MediaIOCore/Public/MediaCapture.h
--- a/Engine/Plugins/Media/MediaIOFramework/Source/MediaIOCore/Public/MediaCapture.h
+++ b/Engine/Plugins/Media/MediaIOFramework/Source/MediaIOCore/Public/MediaCapture.h
@@ -11,6 +11,7 @@
 #include "Misc/Timecode.h"
 #include "PixelFormat.h"
 #include "RendererInterface.h"
+#include "RenderGraphResources.h"
 #include "RHI.h"
 #include "RHIResources.h"
 #include "Templates/Atomic.h"
@@ -21,6 +22,14 @@ class FSceneViewport;
 class FTextureRenderTargetResource;
 class UTextureRenderTarget2D;
 class FRHICommandListImmediate;
+class FRHIGPUTextureReadback;
+class FRHIGPUBufferReadback;
+
+namespace UE::MediaCaptureData
+{
+	class FCaptureFrame;
+	class FMediaCaptureHelper;
+}
 
 /**
  * Possible states of media capture.
@@ -44,11 +53,26 @@ enum class EMediaCaptureState : uint8
 	Stopped,
 };
 
+/**
+ * Possible resource type the MediaCapture can create based on the conversion operation.
+ */
+UENUM()
+enum class EMediaCaptureResourceType : uint8
+{
+	/** Texture resources are used for the readback. */
+	Texture,
+
+	/** RWBuffer resources are used for the readback. */
+	Buffer,
+};
+
 /**
  * Base class of additional data that can be stored for each requested capture.
  */
 class FMediaCaptureUserData
-{};
+{
+
+};
 
 /**
  * Type of cropping 
@@ -147,6 +171,11 @@ class MEDIAIOCORE_API UMediaCapture : public UObject
 	GENERATED_UCLASS_BODY()
 
 public:
+
+	/** Default constructor / destructor to use forward declared uniqueptr */
+	virtual ~UMediaCapture();
+	UMediaCapture(FVTableHelper&);
+
 	/**
 	 * Stop the actual capture if there is one.
 	 * Then start the capture of a SceneViewport.
@@ -246,7 +275,8 @@ protected:
 	virtual bool UpdateRenderTargetImpl(UTextureRenderTarget2D* InRenderTarget) { return true; }
 	virtual void StopCaptureImpl(bool bAllowPendingFrameToBeProcess) { }
 
-
+	friend class UE::MediaCaptureData::FCaptureFrame;
+	friend class UE::MediaCaptureData::FMediaCaptureHelper;
 	struct FCaptureBaseData
 	{
 		FCaptureBaseData();
@@ -263,12 +293,18 @@ protected:
 	 */
 	virtual TSharedPtr<FMediaCaptureUserData, ESPMode::ThreadSafe> GetCaptureFrameUserData_GameThread() { return TSharedPtr<FMediaCaptureUserData, ESPMode::ThreadSafe>(); }
 
-	/** Should we call OnFrameCaptured_RenderingThread() with a RHI Texture -or- copy the memory to CPU ram and call OnFrameCaptured_RenderingThread(). */
+	UE_DEPRECATED(5.1, "ShouldCaptureRHITexture has been deprecated to support both texture and buffer output resource. Please use the ShouldCaptureRHIResource instead.")
 	virtual bool ShouldCaptureRHITexture() const { return false; }
+	
+	/** Should we call OnFrameCaptured_RenderingThread() with a RHI resource -or- readback the memory to CPU ram and call OnFrameCaptured_RenderingThread(). */
+	virtual bool ShouldCaptureRHIResource() const { return false; }
 
-	/** Called at the beginning of the Capture_RenderThread call. */
+	/** Called at the beginning of the Capture_RenderThread call if output resource type is Texture */
 	virtual void BeforeFrameCaptured_RenderingThread(const FCaptureBaseData& InBaseData, TSharedPtr<FMediaCaptureUserData, ESPMode::ThreadSafe> InUserData, FTextureRHIRef InTexture) {}
 
+	/** Called at the beginning of the Capture_RenderThread call if output resource type is Buffer */
+	virtual void BeforeFrameCaptured_RenderingThread(const FCaptureBaseData& InBaseData, TSharedPtr<FMediaCaptureUserData, ESPMode::ThreadSafe> InUserData, FBufferRHIRef InBuffer) {}
+
 	/**
 	 * Callback when the buffer was successfully copied to CPU ram.
 	 * The callback in called from a critical point. If you intend to process the buffer, do so in another thread.
@@ -276,26 +312,68 @@ protected:
 	 */
 	virtual void OnFrameCaptured_RenderingThread(const FCaptureBaseData& InBaseData, TSharedPtr<FMediaCaptureUserData, ESPMode::ThreadSafe> InUserData, void* InBuffer, int32 Width, int32 Height, int32 BytesPerRow) { }
 
+	UE_DEPRECATED(5.1, "OnRHITextureCaptured_RenderingThread has been deprecated to support texture and buffer output resource. Please use OnRHIResourceCaptured_RenderingThread instead.")
+	virtual void OnRHITextureCaptured_RenderingThread(const FCaptureBaseData& InBaseData, TSharedPtr<FMediaCaptureUserData, ESPMode::ThreadSafe> InUserData, FTextureRHIRef InTexture) { }
+
 	/**
-	 * Callback when the buffer was successfully copied on the GPU ram.
-	 * The callback in called from a critical point. If you intend to process the texture, do so in another thread.
+	 * Callbacks when the buffer was successfully copied on the GPU ram.
+	 * The callback in called from a critical point. If you intend to process the texture/buffer, do so in another thread.
 	 * The texture is valid for the duration of the callback.
 	 */
-	virtual void OnRHITextureCaptured_RenderingThread(const FCaptureBaseData& InBaseData, TSharedPtr<FMediaCaptureUserData, ESPMode::ThreadSafe> InUserData, FTextureRHIRef InTexture) { }
+	virtual void OnRHIResourceCaptured_RenderingThread(const FCaptureBaseData& InBaseData, TSharedPtr<FMediaCaptureUserData, ESPMode::ThreadSafe> InUserData, FTextureRHIRef InTexture) { }
+	virtual void OnRHIResourceCaptured_RenderingThread(const FCaptureBaseData& InBaseData, TSharedPtr<FMediaCaptureUserData, ESPMode::ThreadSafe> InUserData, FBufferRHIRef InBuffer) { }
+
+	UE_DEPRECATED(5.1, "OnCustomCapture_RenderingThread has been deprecated while moving to a RDG based pipeline. Please use the RDG version instead.")
+	virtual void OnCustomCapture_RenderingThread(FRHICommandListImmediate& RHICmdList, const FCaptureBaseData& InBaseData, TSharedPtr<FMediaCaptureUserData, ESPMode::ThreadSafe> InUserData, FTextureRHIRef InSourceTexture, FTextureRHIRef TargetableTexture, FResolveParams& ResolveParams, FVector2D CropU, FVector2D CropV)
+	{
+	}
 
 	/**
-	 * Callback when the buffer is ready to be copied on the GPU. You need to copy yourself. You may use the callback to do texel operation or other conversion.
+	 * Callback when the source texture is ready to be copied on the GPU. You need to copy yourself. You may use the callback to add passes to the graph builder.
+	 * Graph execution and output texture readback will be handled after this call. No need to do it in the override.
 	 * The callback in called from a critical point be fast.
 	 * The texture is valid for the duration of the callback.
-	 * Only called when the conversion operation is custom.
+	 * Only called when the conversion operation is custom and the output resource type is Texture.
 	 */
-	virtual void OnCustomCapture_RenderingThread(FRHICommandListImmediate& RHICmdList, const FCaptureBaseData& InBaseData, TSharedPtr<FMediaCaptureUserData, ESPMode::ThreadSafe> InUserData, FTexture2DRHIRef InSourceTexture, FTextureRHIRef TargetableTexture, FResolveParams& ResolveParams, FVector2D CropU, FVector2D CropV) { }
+	virtual void OnCustomCapture_RenderingThread(FRDGBuilder& GraphBuilder, const FCaptureBaseData& InBaseData, TSharedPtr<FMediaCaptureUserData, ESPMode::ThreadSafe> InUserData, FRDGTextureRef InSourceTexture, FRDGTextureRef OutputTexture, FResolveParams& ResolveParams, FVector2D CropU, FVector2D CropV)
+	{
+	}
+
+	/**
+	 * Callback when the source texture is ready to be copied on the GPU. You need to copy yourself. You may use the callback to add passes to the graph builder.
+	 * Graph execution and output buffer readback will be handled after this call. No need to do it in the override.
+	 * The callback in called from a critical point be fast.
+	 * The texture is valid for the duration of the callback.
+	 * Only called when the conversion operation is custom and the output resource type is Buffer.
+	 */
+	virtual void OnCustomCapture_RenderingThread(FRDGBuilder& GraphBuilder, const FCaptureBaseData& InBaseData, TSharedPtr<FMediaCaptureUserData, ESPMode::ThreadSafe> InUserData, FRDGTextureRef InSourceTexture, FRDGBufferRef OutputBuffer, FResolveParams& ResolveParams, FVector2D CropU, FVector2D CropV)
+	{
+	}
 
 	/** Get the size of the buffer when the conversion operation is custom. The function is called once at initialization and the result is cached. */
-	virtual FIntPoint GetCustomOutputSize(const FIntPoint& InSize) const { return InSize; }
+	virtual FIntPoint GetCustomOutputSize(const FIntPoint& InSize) const 
+	{
+		return InSize; 
+	}
 
 	/** Get the pixel format of the buffer when the conversion operation is custom. The function is called once at initialization and the result is cached. */
-	virtual EPixelFormat GetCustomOutputPixelFormat(const EPixelFormat& InPixelFormat) const { return InPixelFormat; }
+	virtual EPixelFormat GetCustomOutputPixelFormat(const EPixelFormat& InPixelFormat) const
+	{
+		return InPixelFormat;
+	}
+	
+	/** Get the desired output resource type when the conversion is custom. This function is called once at initialization and the result is cached */
+	virtual EMediaCaptureResourceType GetCustomOutputResourceType() const 
+	{
+		return DesiredOutputResourceType; 
+	}
+
+	/** Get the desired output buffer description when the conversion is custom and operating in buffer resource type */
+	virtual FRDGBufferDesc GetCustomBufferDescription(const FIntPoint& InDesiredSize) const
+	{
+		return DesiredOutputBufferDescription;
+	}
+
 
 protected:
 	UTextureRenderTarget2D* GetTextureRenderTarget() { return CapturingRenderTarget; }
@@ -303,65 +381,60 @@ protected:
 	EMediaCaptureConversionOperation GetConversionOperation() const { return ConversionOperation; }
 	void SetState(EMediaCaptureState InNewState);
 
-protected:
-	UPROPERTY(Transient)
-	UMediaOutput* MediaOutput;
 
 private:
-	struct FCaptureFrame;
-
-	void InitializeResolveTarget(int32 InNumberOfBuffers);
+	void InitializeOutputResources(int32 InNumberOfBuffers);
 	void OnEndFrame_GameThread();
 	void CacheMediaOutput(EMediaCaptureSourceType InSourceType);
 	void CacheOutputOptions();
 	FIntPoint GetOutputSize(const FIntPoint& InSize, const EMediaCaptureConversionOperation & InConversionOperation) const;
 	EPixelFormat GetOutputPixelFormat(const EPixelFormat& InPixelFormat, const EMediaCaptureConversionOperation & InConversionOperation) const;
+	EMediaCaptureResourceType GetOutputResourceType(const EMediaCaptureConversionOperation & InConversionOperation) const;
+	FRDGBufferDesc GetOutputBufferDescription(EMediaCaptureConversionOperation InConversionOperation) const;
 	void BroadcastStateChanged();
 	void SetFixedViewportSize(TSharedPtr<FSceneViewport> InSceneViewport);
 	void ResetFixedViewportSize(TSharedPtr<FSceneViewport> InViewport, bool bInFlushRenderingCommands);
 
-	void Capture_RenderThread(FRHICommandListImmediate& RHICmdList, UMediaCapture* InMediaCapture, FCaptureFrame* CapturingFrame, FCaptureFrame* ReadyFrame,
+	void Capture_RenderThread(FRHICommandListImmediate& RHICmdList, UMediaCapture* InMediaCapture, UE::MediaCaptureData::FCaptureFrame* CapturingFrame, UE::MediaCaptureData::FCaptureFrame* ReadyFrame,
 		FSceneViewport* InCapturingSceneViewport, FTextureRenderTargetResource* InTextureRenderTargetResource,
 		FIntPoint InDesiredSize, FMediaCaptureStateChangedSignature InOnStateChanged);
 
-private:
-	struct FCaptureFrame
-	{
-		FCaptureFrame();
-
-		FTexture2DRHIRef ReadbackTexture;
-		FCaptureBaseData CaptureBaseData;
-		TAtomic<bool> bResolvedTargetRequested;
-		TSharedPtr<FMediaCaptureUserData, ESPMode::ThreadSafe> UserData;
-		// Used for RHI Capture
-		TRefCountPtr<IPooledRenderTarget> RenderTarget;
-	};
+protected:
 
-	TArray<FCaptureFrame> CaptureFrames;
-	int32 CurrentResolvedTargetIndex;
-	int32 NumberOfCaptureFrame;
-	int32 CaptureRequestCount;
-	EMediaCaptureState MediaState;
+	/** MediaOutput associated with this capture */
+	UPROPERTY(Transient)
+	TObjectPtr<UMediaOutput> MediaOutput;
+	
+	/** Output size of the media capture. If a conversion is done, this resolution might be different than source resolution. i.e. 1080p RGB might be 960x1080 in YUV */
+	FIntPoint DesiredOutputSize = FIntPoint(1920, 1080);
+
+private:
 
 	UPROPERTY(Transient)
-	UTextureRenderTarget2D* CapturingRenderTarget;
+	TObjectPtr<UTextureRenderTarget2D> CapturingRenderTarget;
+	
+	TArray<TUniquePtr<UE::MediaCaptureData::FCaptureFrame>> CaptureFrames;
+	int32 CurrentReadbackIndex = 0;
+	int32 NumberOfCaptureFrame = 2;
+	int32 CaptureRequestCount = 0;
+	EMediaCaptureState MediaState = EMediaCaptureState::Stopped;
+
 	TWeakPtr<FSceneViewport> CapturingSceneViewport;
 	FCriticalSection AccessingCapturingSource;
 
-	FIntPoint DesiredSize;
-	EPixelFormat DesiredPixelFormat;
-	FIntPoint DesiredOutputSize;
-	EPixelFormat DesiredOutputPixelFormat;
+	FIntPoint DesiredSize = FIntPoint(1920, 1080);
+	EPixelFormat DesiredPixelFormat = EPixelFormat::PF_A2B10G10R10;
+	EPixelFormat DesiredOutputPixelFormat = EPixelFormat::PF_A2B10G10R10;
+	FRDGBufferDesc DesiredOutputBufferDescription;
 	FMediaCaptureOptions DesiredCaptureOptions;
-	EMediaCaptureConversionOperation ConversionOperation;
-	FString MediaOutputName;
-	bool bUseRequestedTargetSize;
-
-
-
-	bool bViewportHasFixedViewportSize;
-	TAtomic<bool> bResolvedTargetInitialized;
-	TAtomic<bool> bShouldCaptureRHITexture;
-	TAtomic<int32> WaitingForResolveCommandExecutionCounter;
-	TAtomic<int32> NumberOfTexturesToResolve;
+	EMediaCaptureConversionOperation ConversionOperation = EMediaCaptureConversionOperation::NONE;
+	EMediaCaptureResourceType DesiredOutputResourceType = EMediaCaptureResourceType::Texture;
+	FString MediaOutputName = FString(TEXT("[undefined]"));
+	bool bUseRequestedTargetSize = false;
+
+	bool bViewportHasFixedViewportSize = false;
+	TAtomic<bool> bOutputResourcesInitialized;
+	TAtomic<bool> bShouldCaptureRHIResource;
+	TAtomic<int32> WaitingForRenderCommandExecutionCounter;
+	TAtomic<int32> NumberOfResourcesToReadback;
 };
diff --git a/Engine/Plugins/VirtualProduction/Rivermax/RivermaxCore/RivermaxCore.uplugin b/Engine/Plugins/VirtualProduction/Rivermax/RivermaxCore/RivermaxCore.uplugin
--- a/Engine/Plugins/VirtualProduction/Rivermax/RivermaxCore/RivermaxCore.uplugin
+++ b/Engine/Plugins/VirtualProduction/Rivermax/RivermaxCore/RivermaxCore.uplugin
@@ -27,6 +27,12 @@
 			"Type": "Runtime",
 			"LoadingPhase": "Default",
 			"PlatformAllowList": [ "Win64" ]
+		},
+		{
+			"Name": "RivermaxRendering",
+			"Type": "Runtime",
+			"LoadingPhase": "PostConfigInit",
+			"PlatformAllowList": [ "Win64" ]
 		}
 	]
 }
\ No newline at end of file
diff --git a/Engine/Plugins/VirtualProduction/Rivermax/RivermaxCore/Shaders/Private/RivermaxShaders.usf b/Engine/Plugins/VirtualProduction/Rivermax/RivermaxCore/Shaders/Private/RivermaxShaders.usf
new file mode 100644
--- /dev/null
+++ b/Engine/Plugins/VirtualProduction/Rivermax/RivermaxCore/Shaders/Private/RivermaxShaders.usf
@@ -0,0 +1,190 @@
+// Copyright Epic Games, Inc. All Rights Reserved.
+
+/*=============================================================================
+	RivermaxShaders.usf
+=============================================================================*/
+
+#include "/Engine/Public/Platform.ush"
+#include "/Engine/Private/ColorUtils.ush"
+#include "/Engine/Private/GammaCorrectionCommon.ush"
+
+Texture2D InputTexture;
+SamplerState InputSampler;
+float4x4 ColorTransform;
+uint DoLinearToSrgb;
+float OnePixelDeltaX;
+float OnePixelDeltaY;
+float InputPixelOffsetX;
+float InputPixelOffsetY;
+uint HorizontalElementCount;
+uint InputTexturePixelsPerThread;
+
+
+struct FYCbCr10Bit422
+{
+	uint DWord0;
+	uint DWord1;
+	uint DWord2;
+	uint DWord3;
+	uint DWord4;
+};
+RWStructuredBuffer<FYCbCr10Bit422>	OutYCbCrBuffer;
+
+struct FRGB8BitBuffer
+{
+	uint DWord0;
+	uint DWord1;
+	uint DWord2;
+};
+RWStructuredBuffer<FRGB8BitBuffer>	OutRGB8Buffer;
+
+struct FRGB10BitBuffer
+{
+	uint DWords[15];
+};
+RWStructuredBuffer<FRGB10BitBuffer>	OutRGB10Buffer;
+
+struct FRGB16fBuffer
+{
+	half Pixel0[3];
+	half Pixel1[3];
+};
+RWStructuredBuffer<FRGB16fBuffer> OutRGB16fBuffer;
+
+// Shader to convert from RGB to YUV 422 10 bits packed little endian
+
+[numthreads(8,8,1)]
+void RGBToYUV10Bit422(uint3 ThreadId : SV_DispatchThreadID, uint GroupIndex : SV_GroupIndex)
+{
+	uint Index = ThreadId.y * HorizontalElementCount + ThreadId.x;
+
+	float X = (ThreadId.x * InputTexturePixelsPerThread) + InputPixelOffsetX;
+	float Y = (ThreadId.y + InputPixelOffsetY) * OnePixelDeltaY;
+
+	// We process 8 RGB pixels per thread
+	float3 RGB0 = InputTexture.SampleLevel(InputSampler, float2(((X + 0.0) + 0.5) * OnePixelDeltaX, Y), 0).rgb;
+	float3 RGB1 = InputTexture.SampleLevel(InputSampler, float2(((X + 1.0) + 0.5) * OnePixelDeltaX, Y), 0).rgb;
+	float3 RGB2 = InputTexture.SampleLevel(InputSampler, float2(((X + 2.0) + 0.5) * OnePixelDeltaX, Y), 0).rgb;
+	float3 RGB3 = InputTexture.SampleLevel(InputSampler, float2(((X + 3.0) + 0.5) * OnePixelDeltaX, Y), 0).rgb;
+	float3 RGB4 = InputTexture.SampleLevel(InputSampler, float2(((X + 4.0) + 0.5) * OnePixelDeltaX, Y), 0).rgb;
+	float3 RGB5 = InputTexture.SampleLevel(InputSampler, float2(((X + 5.0) + 0.5) * OnePixelDeltaX, Y), 0).rgb;
+	float3 RGB6 = InputTexture.SampleLevel(InputSampler, float2(((X + 6.0) + 0.5) * OnePixelDeltaX, Y), 0).rgb;
+	float3 RGB7 = InputTexture.SampleLevel(InputSampler, float2(((X + 7.0) + 0.5) * OnePixelDeltaX, Y), 0).rgb;
+
+	// The 8 RGB needs to be converted to YUV. Each encoded pixel will use 2 RGB pixels (YnUnVnYn+1)
+	uint3 YUV0 = RgbToYuv(RGB0, ColorTransform, DoLinearToSrgb) * 1023;
+	uint3 YUV1 = RgbToYuv(RGB1, ColorTransform, DoLinearToSrgb) * 1023;
+	uint3 YUV2 = RgbToYuv(RGB2, ColorTransform, DoLinearToSrgb) * 1023;
+	uint3 YUV3 = RgbToYuv(RGB3, ColorTransform, DoLinearToSrgb) * 1023;
+	uint3 YUV4 = RgbToYuv(RGB4, ColorTransform, DoLinearToSrgb) * 1023;
+	uint3 YUV5 = RgbToYuv(RGB5, ColorTransform, DoLinearToSrgb) * 1023;
+	uint3 YUV6 = RgbToYuv(RGB6, ColorTransform, DoLinearToSrgb) * 1023;
+	uint3 YUV7 = RgbToYuv(RGB7, ColorTransform, DoLinearToSrgb) * 1023;
+
+	// We write the 4 YUV pixels (40 bits * 4 = 160 bits = 5x32 bits) following the little endian packing 2110 expects
+	OutYCbCrBuffer[Index].DWord0 = ((YUV0.z & 0x3F) << 26) | ((YUV1.x >> 8) << 24) | ((YUV0.x & 0xF) << 20) | ((YUV0.z >> 6) << 16) | ((YUV0.y & 0x3) << 14) | ((YUV0.x >> 4) << 8) | (YUV0.y >> 2);
+	OutYCbCrBuffer[Index].DWord1 = ((YUV2.x & 0xF) << 28) | ((YUV2.z >> 6) << 24) | ((YUV2.y & 0x3) << 22) | ((YUV2.x >> 4) << 16) | ((YUV2.y >> 2) << 8) | (YUV1.x & 0xFF);
+	OutYCbCrBuffer[Index].DWord2 = ((YUV4.y & 0x3) << 30) | ((YUV4.x >> 4) << 24) | ((YUV4.y >> 2) << 16) | ((YUV3.x & 0xFF) << 8) | ((YUV2.z & 0x3F) << 2) | (YUV3.x >> 8);
+	OutYCbCrBuffer[Index].DWord3 = ((YUV6.y >> 2) << 24) | ((YUV5.x & 0xFF) << 16) | ((YUV4.z & 0x3F) << 10) | ((YUV5.x >> 8) << 8) | ((YUV4.x & 0xF) << 4) | (YUV4.z >> 6);
+	OutYCbCrBuffer[Index].DWord4 = ((YUV7.x & 0xFF) << 24) | ((YUV6.z & 0x3F) << 18) | ((YUV7.x >> 8) << 16) | ((YUV6.x & 0xF) << 12) | ((YUV6.z >> 6) << 8) | ((YUV6.y & 0x3) << 6) | (YUV6.x >> 4);
+}
+
+
+// Shader to convert from RGBA to RGB8
+
+[numthreads(8,8,1)]
+void RGBToRGB8BitCS(uint3 ThreadId : SV_DispatchThreadID, uint GroupIndex : SV_GroupIndex)
+{
+	uint Index = ThreadId.y * HorizontalElementCount + ThreadId.x;
+
+	float X = (ThreadId.x * InputTexturePixelsPerThread) + InputPixelOffsetX;
+	float Y = (ThreadId.y + InputPixelOffsetY) * OnePixelDeltaY;
+
+	// We process 4 RGB pixels per thread
+	uint3 RGB0 = InputTexture.SampleLevel(InputSampler, float2(((X + 0.0) + 0.5) * OnePixelDeltaX, Y), 0).rgb * 255;
+	uint3 RGB1 = InputTexture.SampleLevel(InputSampler, float2(((X + 1.0) + 0.5) * OnePixelDeltaX, Y), 0).rgb * 255;
+	uint3 RGB2 = InputTexture.SampleLevel(InputSampler, float2(((X + 2.0) + 0.5) * OnePixelDeltaX, Y), 0).rgb * 255;
+	uint3 RGB3 = InputTexture.SampleLevel(InputSampler, float2(((X + 3.0) + 0.5) * OnePixelDeltaX, Y), 0).rgb * 255;
+
+	// We write the 4 RGB pixels (24 bits * 4 = 96 bits = 3x32 bits) following the packing 2110 expects
+	OutRGB8Buffer[Index].DWord0 = (RGB1.x << 24) | (RGB0.z << 16) | (RGB0.y << 8) | (RGB0.x << 0);
+	OutRGB8Buffer[Index].DWord1 = (RGB2.y << 24) | (RGB2.x << 16) | (RGB1.z << 8) | (RGB1.y << 0);
+	OutRGB8Buffer[Index].DWord2 = (RGB3.z << 24) | (RGB3.y << 16) | (RGB3.x << 8) | (RGB2.z << 0);
+}
+
+// Shader to convert from RGBA to RGB10
+
+[numthreads(8,8,1)]
+void RGBToRGB10BitCS(uint3 ThreadId : SV_DispatchThreadID, uint GroupIndex : SV_GroupIndex)
+{
+	uint Index = ThreadId.y * HorizontalElementCount + ThreadId.x;
+
+	float X = (ThreadId.x * InputTexturePixelsPerThread) + InputPixelOffsetX;
+	float Y = (ThreadId.y + InputPixelOffsetY) * OnePixelDeltaY;
+
+	uint3 RGB00 = InputTexture.SampleLevel(InputSampler, float2(((X + 0.0) + 0.5) * OnePixelDeltaX, Y), 0).rgb * 1023;
+	uint3 RGB01 = InputTexture.SampleLevel(InputSampler, float2(((X + 1.0) + 0.5) * OnePixelDeltaX, Y), 0).rgb * 1023;
+	uint3 RGB02 = InputTexture.SampleLevel(InputSampler, float2(((X + 2.0) + 0.5) * OnePixelDeltaX, Y), 0).rgb * 1023;
+	uint3 RGB03 = InputTexture.SampleLevel(InputSampler, float2(((X + 3.0) + 0.5) * OnePixelDeltaX, Y), 0).rgb * 1023;
+	uint3 RGB04 = InputTexture.SampleLevel(InputSampler, float2(((X + 4.0) + 0.5) * OnePixelDeltaX, Y), 0).rgb * 1023;
+	uint3 RGB05 = InputTexture.SampleLevel(InputSampler, float2(((X + 5.0) + 0.5) * OnePixelDeltaX, Y), 0).rgb * 1023;
+	uint3 RGB06 = InputTexture.SampleLevel(InputSampler, float2(((X + 6.0) + 0.5) * OnePixelDeltaX, Y), 0).rgb * 1023;
+	uint3 RGB07 = InputTexture.SampleLevel(InputSampler, float2(((X + 7.0) + 0.5) * OnePixelDeltaX, Y), 0).rgb * 1023;
+	uint3 RGB08 = InputTexture.SampleLevel(InputSampler, float2(((X + 8.0) + 0.5) * OnePixelDeltaX, Y), 0).rgb * 1023;
+	uint3 RGB09 = InputTexture.SampleLevel(InputSampler, float2(((X + 9.0) + 0.5) * OnePixelDeltaX, Y), 0).rgb * 1023;
+	uint3 RGB10 = InputTexture.SampleLevel(InputSampler, float2(((X + 10.0) + 0.5) * OnePixelDeltaX, Y), 0).rgb * 1023;
+	uint3 RGB11 = InputTexture.SampleLevel(InputSampler, float2(((X + 11.0) + 0.5) * OnePixelDeltaX, Y), 0).rgb * 1023;
+	uint3 RGB12 = InputTexture.SampleLevel(InputSampler, float2(((X + 12.0) + 0.5) * OnePixelDeltaX, Y), 0).rgb * 1023;
+	uint3 RGB13 = InputTexture.SampleLevel(InputSampler, float2(((X + 13.0) + 0.5) * OnePixelDeltaX, Y), 0).rgb * 1023;
+	uint3 RGB14 = InputTexture.SampleLevel(InputSampler, float2(((X + 14.0) + 0.5) * OnePixelDeltaX, Y), 0).rgb * 1023;
+	uint3 RGB15 = InputTexture.SampleLevel(InputSampler, float2(((X + 15.0) + 0.5) * OnePixelDeltaX, Y), 0).rgb * 1023;
+
+	// All this bit swizzling is to pack 3 channels of 10 bits with R being in LSB of the first 32 bits and subsequent channels following until we used all 16 pixels.
+	// Example : RGB = (0x123, 0x245, 0x0AB)
+	// Binary  : 0100 1000 1110 0100 0101 0011 0011 0100
+	// Bytes   : 0x48 0xE4 0x53 0x34
+	// 32 bits : 34 << 24 | 53 << 16 | E4 << 8 | 48 << 0
+	// The bits of each channel not being next to each other seems weird but when you look at the byte ordering you can see it
+	// The E in the second byte is the LSB of R and it's separated from the rest of R by the MSB of G.
+	// We want 48 at the first adress so it goes in LSB of the 32 bits as it's written as little endian. 
+	// 2110 on the wire will read it as big endian which will give us R going out first
+
+	OutRGB10Buffer[Index].DWords[0] = ((RGB00.z & 0x03F) << 26) | ((RGB01.x >> 8)   << 24)   | ((RGB00.y & 0xF)  << 20) | ((RGB00.z >> 6)   << 16) | ((RGB00.x & 0x3)  << 14) | ((RGB00.y >> 4)   << 8)								| (RGB00.x >> 2);
+	OutRGB10Buffer[Index].DWords[1] = ((RGB01.z & 0x00F) << 28) | ((RGB02.x >> 6)   << 24)   | ((RGB01.y & 0x3)  << 22) | ((RGB01.z >> 4)   << 16)							  | ((RGB01.y >> 2)   << 8)								| (RGB01.x & 0xFF);
+	OutRGB10Buffer[Index].DWords[2] = ((RGB02.z & 0x003) << 30) | ((RGB03.x >> 4)   << 24)								| ((RGB02.z >> 2)   << 16)							  | ((RGB02.y & 0xFF) << 8)  | ((RGB02.x & 0x3F) << 2)  | (RGB02.y >> 8);
+	OutRGB10Buffer[Index].DWords[3] =							  ((RGB04.x >> 2)   << 24)								| ((RGB03.z & 0xFF) << 16) | ((RGB03.y & 0x3F) << 10) | ((RGB03.z >> 8)   << 8)  | ((RGB03.x & 0xF)  << 4)  | (RGB03.y >> 6);
+	OutRGB10Buffer[Index].DWords[4] =							  ((RGB05.x & 0xFF) << 24)   | ((RGB04.z & 0x3F) << 18) | ((RGB05.x >> 8)   << 16) | ((RGB04.y & 0xF)  << 12) | ((RGB04.z >> 6)   << 8)  | ((RGB04.x & 0x3)  << 6)  | (RGB04.y >> 4);
+	
+	OutRGB10Buffer[Index].DWords[5] = ((RGB05.x & 0x03F) << 26) | ((RGB06.y >> 8)   << 24)   | ((RGB05.z & 0xF)  << 20) | ((RGB05.x >> 6)   << 16) | ((RGB05.y & 0x3)  << 14) | ((RGB05.z >> 4)   << 8)								| (RGB05.y >> 2);
+	OutRGB10Buffer[Index].DWords[6] = ((RGB06.x & 0x00F) << 28) | ((RGB07.y >> 6)   << 24)   | ((RGB06.z & 0x3)  << 22) | ((RGB06.x >> 4)   << 16)							  | ((RGB06.z >> 2)   << 8)								| (RGB06.y & 0xFF);
+	OutRGB10Buffer[Index].DWords[7] = ((RGB07.x & 0x003) << 30) | ((RGB08.y >> 4)   << 24)								| ((RGB07.x >> 2)   << 16)							  | ((RGB07.z & 0xFF) << 8)  | ((RGB07.y & 0x3F) << 2)  | (RGB07.z >> 8);
+	OutRGB10Buffer[Index].DWords[8] =							  ((RGB09.y >> 2)   << 24)								| ((RGB08.x & 0xFF) << 16) | ((RGB08.z & 0x3F) << 10) | ((RGB08.x >> 8)   << 8)  | ((RGB08.y & 0xF)  << 4)  | (RGB08.z >> 6);
+	OutRGB10Buffer[Index].DWords[9] =							  ((RGB10.y & 0xFF) << 24)   | ((RGB09.x & 0x3F) << 18) | ((RGB10.y >> 8)   << 16) | ((RGB09.z & 0xF)  << 12) | ((RGB09.x >> 6)   << 8)  | ((RGB09.y & 0x3)  << 6)  | (RGB09.z >> 4);
+
+	OutRGB10Buffer[Index].DWords[10] = ((RGB10.y & 0x03F) << 26) | ((RGB11.z >> 8)   << 24)  | ((RGB11.x & 0xF)  << 20) | ((RGB11.y >> 6)   << 16) | ((RGB10.z & 0x3)  << 14) | ((RGB10.x >> 4)   << 8)								| (RGB10.z >> 2);
+	OutRGB10Buffer[Index].DWords[11] = ((RGB11.y & 0x00F) << 28) | ((RGB12.z >> 6)   << 24)  | ((RGB12.x & 0x3)  << 22) | ((RGB12.y >> 4)   << 16)							  | ((RGB11.x >> 2)   << 8)								| (RGB11.z & 0xFF);
+	OutRGB10Buffer[Index].DWords[12] = ((RGB12.y & 0x003) << 30) | ((RGB13.z >> 4)   << 24)								| ((RGB13.y >> 2)   << 16)							  | ((RGB12.x & 0xFF) << 8)  | ((RGB12.z & 0x3F) << 2)  | (RGB12.x >> 8);
+	OutRGB10Buffer[Index].DWords[13] =							   ((RGB14.z >> 2)   << 24)								| ((RGB14.y & 0xFF) << 16) | ((RGB13.x & 0x3F) << 10) | ((RGB13.y >> 8)   << 8)  | ((RGB13.z & 0xF)  << 4)  | (RGB13.x >> 6);
+	OutRGB10Buffer[Index].DWords[14] =							   ((RGB15.z & 0xFF) << 24)  | ((RGB14.y & 0x3F) << 18) | ((RGB15.z >> 8)   << 16) | ((RGB14.x & 0xF)  << 12) | ((RGB14.y >> 6)   << 8)  | ((RGB14.z & 0x3)  << 6)  | (RGB14.x >> 4);
+}
+
+// Shader to convert from RGBA to RGB Float16
+
+[numthreads(8,8,1)]
+void RGBToRGB16fBitCS(uint3 ThreadId : SV_DispatchThreadID, uint GroupIndex : SV_GroupIndex)
+{
+	float X = (ThreadId.x * InputTexturePixelsPerThread) + InputPixelOffsetX;
+	float Y = (ThreadId.y + InputPixelOffsetY) * OnePixelDeltaY;
+
+	half3 RGB0 = InputTexture.SampleLevel(InputSampler, float2(((X + 0.0) + 0.5) * OnePixelDeltaX, Y), 0).rgb;
+	half3 RGB1 = InputTexture.SampleLevel(InputSampler, float2(((X + 1.0) + 0.5) * OnePixelDeltaX, Y), 0).rgb;
+	
+	// Storage has been untested as no available supported that viewer. Will need to be validated. 
+	uint Index = ThreadId.y * HorizontalElementCount + ThreadId.x;
+	OutRGB16fBuffer[Index].Pixel0[0]  = RGB0.z;
+	OutRGB16fBuffer[Index].Pixel0[1]  = RGB0.y;
+	OutRGB16fBuffer[Index].Pixel0[2]  = RGB0.x;
+	OutRGB16fBuffer[Index].Pixel1[0]  = RGB1.z;
+	OutRGB16fBuffer[Index].Pixel1[1]  = RGB1.y;
+	OutRGB16fBuffer[Index].Pixel1[2]  = RGB1.x;
+}
diff --git a/Engine/Plugins/VirtualProduction/Rivermax/RivermaxCore/Source/Private/RivermaxCoreModule.cpp b/Engine/Plugins/VirtualProduction/Rivermax/RivermaxCore/Source/Private/RivermaxCoreModule.cpp
deleted file mode 100644
diff --git a/Engine/Plugins/VirtualProduction/Rivermax/RivermaxCore/Source/Private/RivermaxCoreModule.h b/Engine/Plugins/VirtualProduction/Rivermax/RivermaxCore/Source/Private/RivermaxCoreModule.h
deleted file mode 100644
diff --git a/Engine/Plugins/VirtualProduction/Rivermax/RivermaxCore/Source/Private/RivermaxInputStream.cpp b/Engine/Plugins/VirtualProduction/Rivermax/RivermaxCore/Source/Private/RivermaxInputStream.cpp
deleted file mode 100644
diff --git a/Engine/Plugins/VirtualProduction/Rivermax/RivermaxCore/Source/Private/RivermaxInputStream.h b/Engine/Plugins/VirtualProduction/Rivermax/RivermaxCore/Source/Private/RivermaxInputStream.h
deleted file mode 100644
diff --git a/Engine/Plugins/VirtualProduction/Rivermax/RivermaxCore/Source/Private/RivermaxLog.cpp b/Engine/Plugins/VirtualProduction/Rivermax/RivermaxCore/Source/Private/RivermaxLog.cpp
deleted file mode 100644
diff --git a/Engine/Plugins/VirtualProduction/Rivermax/RivermaxCore/Source/Private/RivermaxLog.h b/Engine/Plugins/VirtualProduction/Rivermax/RivermaxCore/Source/Private/RivermaxLog.h
deleted file mode 100644
diff --git a/Engine/Plugins/VirtualProduction/Rivermax/RivermaxCore/Source/Private/RivermaxManager.cpp b/Engine/Plugins/VirtualProduction/Rivermax/RivermaxCore/Source/Private/RivermaxManager.cpp
deleted file mode 100644
diff --git a/Engine/Plugins/VirtualProduction/Rivermax/RivermaxCore/Source/Private/RivermaxManager.h b/Engine/Plugins/VirtualProduction/Rivermax/RivermaxCore/Source/Private/RivermaxManager.h
deleted file mode 100644
diff --git a/Engine/Plugins/VirtualProduction/Rivermax/RivermaxCore/Source/Private/RivermaxOutputFrame.cpp b/Engine/Plugins/VirtualProduction/Rivermax/RivermaxCore/Source/Private/RivermaxOutputFrame.cpp
deleted file mode 100644
diff --git a/Engine/Plugins/VirtualProduction/Rivermax/RivermaxCore/Source/Private/RivermaxOutputFrame.h b/Engine/Plugins/VirtualProduction/Rivermax/RivermaxCore/Source/Private/RivermaxOutputFrame.h
deleted file mode 100644
diff --git a/Engine/Plugins/VirtualProduction/Rivermax/RivermaxCore/Source/Private/RivermaxOutputStream.cpp b/Engine/Plugins/VirtualProduction/Rivermax/RivermaxCore/Source/Private/RivermaxOutputStream.cpp
deleted file mode 100644
diff --git a/Engine/Plugins/VirtualProduction/Rivermax/RivermaxCore/Source/Private/RivermaxOutputStream.h b/Engine/Plugins/VirtualProduction/Rivermax/RivermaxCore/Source/Private/RivermaxOutputStream.h
deleted file mode 100644
diff --git a/Engine/Plugins/VirtualProduction/Rivermax/RivermaxCore/Source/Private/RivermaxUtils.cpp b/Engine/Plugins/VirtualProduction/Rivermax/RivermaxCore/Source/Private/RivermaxUtils.cpp
deleted file mode 100644
diff --git a/Engine/Plugins/VirtualProduction/Rivermax/RivermaxCore/Source/Private/RivermaxUtils.h b/Engine/Plugins/VirtualProduction/Rivermax/RivermaxCore/Source/Private/RivermaxUtils.h
deleted file mode 100644
diff --git a/Engine/Plugins/VirtualProduction/Rivermax/RivermaxCore/Source/Public/IRivermaxCoreModule.h b/Engine/Plugins/VirtualProduction/Rivermax/RivermaxCore/Source/Public/IRivermaxCoreModule.h
deleted file mode 100644
diff --git a/Engine/Plugins/VirtualProduction/Rivermax/RivermaxCore/Source/Public/IRivermaxInputStream.h b/Engine/Plugins/VirtualProduction/Rivermax/RivermaxCore/Source/Public/IRivermaxInputStream.h
deleted file mode 100644
diff --git a/Engine/Plugins/VirtualProduction/Rivermax/RivermaxCore/Source/Public/IRivermaxManager.h b/Engine/Plugins/VirtualProduction/Rivermax/RivermaxCore/Source/Public/IRivermaxManager.h
deleted file mode 100644
diff --git a/Engine/Plugins/VirtualProduction/Rivermax/RivermaxCore/Source/Public/IRivermaxOutputStream.h b/Engine/Plugins/VirtualProduction/Rivermax/RivermaxCore/Source/Public/IRivermaxOutputStream.h
deleted file mode 100644
diff --git a/Engine/Plugins/VirtualProduction/Rivermax/RivermaxCore/Source/Public/RivermaxHeader.h b/Engine/Plugins/VirtualProduction/Rivermax/RivermaxCore/Source/Public/RivermaxHeader.h
deleted file mode 100644
diff --git a/Engine/Plugins/VirtualProduction/Rivermax/RivermaxCore/Source/Public/RivermaxTypes.h b/Engine/Plugins/VirtualProduction/Rivermax/RivermaxCore/Source/Public/RivermaxTypes.h
deleted file mode 100644
diff --git a/Engine/Plugins/VirtualProduction/Rivermax/RivermaxCore/Source/RivermaxCore.build.cs b/Engine/Plugins/VirtualProduction/Rivermax/RivermaxCore/Source/RivermaxCore.build.cs
deleted file mode 100644
diff --git a/Engine/Plugins/VirtualProduction/Rivermax/RivermaxCore/Source/RivermaxCore/Private/RivermaxCoreModule.cpp b/Engine/Plugins/VirtualProduction/Rivermax/RivermaxCore/Source/RivermaxCore/Private/RivermaxCoreModule.cpp
new file mode 100644
--- /dev/null
+++ b/Engine/Plugins/VirtualProduction/Rivermax/RivermaxCore/Source/RivermaxCore/Private/RivermaxCoreModule.cpp
@@ -0,0 +1,38 @@
+// Copyright Epic Games, Inc. All Rights Reserved.
+
+#include "RivermaxCoreModule.h"
+
+#include "Modules/ModuleManager.h"
+#include "RivermaxInputStream.h"
+#include "RivermaxManager.h"
+#include "RivermaxOutputStream.h"
+
+
+void FRivermaxCoreModule::StartupModule()
+{
+	RivermaxManager = MakeShared<UE::RivermaxCore::Private::FRivermaxManager>();
+}
+
+void FRivermaxCoreModule::ShutdownModule()
+{
+
+}
+
+TUniquePtr<UE::RivermaxCore::IRivermaxInputStream> FRivermaxCoreModule::CreateInputStream()
+{
+	using UE::RivermaxCore::Private::FRivermaxInputStream;
+	return MakeUnique<FRivermaxInputStream>();
+}
+
+TUniquePtr<UE::RivermaxCore::IRivermaxOutputStream> FRivermaxCoreModule::CreateOutputStream()
+{
+	using UE::RivermaxCore::Private::FRivermaxOutputStream;
+	return MakeUnique<FRivermaxOutputStream>();
+}
+
+TSharedPtr<UE::RivermaxCore::IRivermaxManager> FRivermaxCoreModule::GetRivermaxManager()
+{
+	return RivermaxManager;
+}
+
+IMPLEMENT_MODULE(FRivermaxCoreModule, RivermaxCore);
diff --git a/Engine/Plugins/VirtualProduction/Rivermax/RivermaxCore/Source/RivermaxCore/Private/RivermaxCoreModule.h b/Engine/Plugins/VirtualProduction/Rivermax/RivermaxCore/Source/RivermaxCore/Private/RivermaxCoreModule.h
new file mode 100644
--- /dev/null
+++ b/Engine/Plugins/VirtualProduction/Rivermax/RivermaxCore/Source/RivermaxCore/Private/RivermaxCoreModule.h
@@ -0,0 +1,35 @@
+// Copyright Epic Games, Inc. All Rights Reserved.
+
+#pragma once
+
+#include "IRivermaxCoreModule.h"
+
+
+namespace UE::RivermaxCore
+{
+	class IRivermaxManager;
+}
+
+/**
+ * Implementation of the IRivermaxCoreModule managing Rivermax library and stream creation
+ */
+class FRivermaxCoreModule : public IRivermaxCoreModule
+{
+
+public:
+	//~ Begin IModuleInterface
+	virtual void StartupModule() override;
+	virtual void ShutdownModule() override;
+	//~ End IModuleInterface
+
+	//~ Begin IRivermaxCoreModule interface
+	virtual TUniquePtr<UE::RivermaxCore::IRivermaxInputStream> CreateInputStream() override;
+	virtual TUniquePtr<UE::RivermaxCore::IRivermaxOutputStream> CreateOutputStream() override;
+	virtual TSharedPtr<UE::RivermaxCore::IRivermaxManager> GetRivermaxManager() override;
+	//~ End IRivermaxCoreModule interface
+
+private:
+
+	TSharedPtr<UE::RivermaxCore::IRivermaxManager> RivermaxManager;
+};
+
diff --git a/Engine/Plugins/VirtualProduction/Rivermax/RivermaxCore/Source/RivermaxCore/Private/RivermaxInputStream.cpp b/Engine/Plugins/VirtualProduction/Rivermax/RivermaxCore/Source/RivermaxCore/Private/RivermaxInputStream.cpp
new file mode 100644
--- /dev/null
+++ b/Engine/Plugins/VirtualProduction/Rivermax/RivermaxCore/Source/RivermaxCore/Private/RivermaxInputStream.cpp
@@ -0,0 +1,443 @@
+// Copyright Epic Games, Inc. All Rights Reserved.
+
+#include "RivermaxInputStream.h"
+
+#include "Async/Async.h"
+#include "IRivermaxCoreModule.h"
+#include "IRivermaxManager.h"
+#include "Misc/ByteSwap.h"
+#include "RivermaxLog.h"
+#include "RivermaxUtils.h"
+
+#if PLATFORM_WINDOWS
+#include <WS2tcpip.h>
+#endif
+
+
+namespace UE::RivermaxCore::Private
+{
+#if PLATFORM_SUPPORTS_PRAGMA_PACK
+#pragma pack(push, 1)
+#endif
+	struct FRivermaxRTPSampleRowData
+	{
+		uint32 ContributingSourceCount : 4;
+		uint32 ExtensionBit : 1;
+		uint32 PaddingBit : 1;
+		uint32 Version : 2;
+		uint32 PayloadType : 7;
+		uint32 MarkerBit : 1;
+		uint32 SequenceNumber : 16;
+		uint32 Timestamp : 32;
+		uint32 SynchronizationSource : 32;
+		uint32 ExtendedSequenceNumber : 16;
+		//SRD 1
+		uint32 SRDLength1 : 16;
+		uint32 SRDRowNumberHigh1 : 7;
+		uint32 FieldIdentification1 : 1;
+		uint32 SRDRowNumberLow1 : 8;
+		uint32 SRDOffsetHigh1 : 7;
+		uint32 ContinuationBit1 : 1;
+		uint32 SRDOffsetLow1 : 8;
+		//SRD 2
+		uint32 SRDLength2 : 16;
+		uint32 SRDRowNumberHigh2 : 7;
+		uint32 FieldIdentification2 : 1;
+		uint32 SRDRowNumberLow2 : 8;
+		uint32 SRDOffsetHigh2 : 7;
+		uint32 ContinuationBit2 : 1;
+		uint32 SRDOffsetLow2 : 8;
+
+		uint16 GetSrd1RowNumber() { return ((SRDRowNumberHigh1 << 8) | SRDRowNumberLow1); }
+		uint16 GetSrd1Offset() { return ((SRDOffsetHigh1 << 8) | SRDOffsetLow1); }
+
+		uint16 GetSrd2RowNumber() { return ((SRDRowNumberHigh2 << 8) | SRDRowNumberLow2); }
+		uint16 GetSrd2Offset() { return ((SRDOffsetHigh2 << 8) | SRDOffsetLow2); }
+	};
+#if PLATFORM_SUPPORTS_PRAGMA_PACK
+#pragma pack(pop)
+#endif
+
+
+	uint8* GetRTPHeaderPointer(uint8* InHeader)
+	{
+		if (!InHeader /*|| RMAX_APP_PROTOCOL_PACKET == m_type*/)
+		{
+			return nullptr;
+		}
+
+		static constexpr uint32 ETH_TYPE_802_1Q = 0x8100;          /* 802.1Q VLAN Extended Header  */
+		static constexpr uint32 RTP_HEADER_SIZE = 12;
+		uint16* ETHProto = (uint16_t*)(InHeader + RTP_HEADER_SIZE);
+		if (ETH_TYPE_802_1Q == ByteSwap(*ETHProto))
+		{
+			InHeader += 46; // 802 + 802.1Q + IP + UDP
+		}
+		else
+		{
+			InHeader += 42; // 802 + IP + UDP
+		}
+		return InHeader;
+	}
+
+	FRivermaxInputStream::FRivermaxInputStream()
+	{
+
+	}
+
+	FRivermaxInputStream::~FRivermaxInputStream()
+	{
+		Uninitialize();
+	}
+
+	bool FRivermaxInputStream::Initialize(const FRivermaxStreamOptions& InOptions, IRivermaxInputStreamListener& InListener)
+	{
+		IRivermaxCoreModule& RivermaxModule = FModuleManager::LoadModuleChecked<IRivermaxCoreModule>(TEXT("RivermaxCore"));
+		if (RivermaxModule.GetRivermaxManager()->IsInitialized() == false)
+		{
+			UE_LOG(LogRivermax, Warning, TEXT("Can't create Rivermax Input Stream. Library isn't initialized."));
+			return false;
+		}
+
+		Options = InOptions;
+		Listener = &InListener;
+
+		InitTaskFuture = Async(EAsyncExecution::TaskGraph, [this]()
+		{
+			// If the stream is trying to shutdown before the init task has even started, don't bother
+			if (bIsShuttingDown)
+			{
+					return;
+			}
+
+			bool bWasSuccessful = false;
+			int32 FlowId = 0; //todo configure
+
+			//Configure local IP interface
+			const rmax_in_stream_type StreamType = RMAX_RAW_PACKET;
+			sockaddr_in RivermaxInterface;
+			memset(&RivermaxInterface, 0, sizeof(RivermaxInterface));
+			if (inet_pton(AF_INET, StringCast<ANSICHAR>(*Options.SourceAddress).Get(), &RivermaxInterface.sin_addr) != 1)
+			{
+				UE_LOG(LogRivermax, Warning, TEXT("inet_pton failed to %s"), *Options.SourceAddress);
+			}
+			else
+			{
+				RivermaxInterface.sin_family = AF_INET;
+
+				//Configure Flow and destination IP (multicast)
+				memset(&FlowAttribute, 0, sizeof(FlowAttribute));
+				FlowAttribute.local_addr.sin_family = AF_INET;
+				FlowAttribute.flow_id = FlowId;
+				if (inet_pton(AF_INET, StringCast<ANSICHAR>(*Options.DestinationAddress).Get(), &FlowAttribute.local_addr.sin_addr) != 1)
+				{
+					UE_LOG(LogRivermax, Warning, TEXT("inet_pton failed to %s"), *Options.DestinationAddress);
+				}
+				else
+				{
+					FlowAttribute.local_addr.sin_port = ByteSwap((uint16)Options.Port);
+
+					const rmax_in_buffer_attr_flags_t BufferAttributeFlags = RMAX_IN_BUFFER_ATTER_FLAG_NONE; //todo whether ordering is based on sequence or extended sequence
+					uint32 BufferElement = 1 << 18;//todo number of packets to allocate memory for
+					rmax_in_buffer_attr BufferAttributes;
+					FMemory::Memset(&BufferAttributes, 0, sizeof(BufferAttributes));
+					BufferAttributes.num_of_elements = BufferElement;
+					BufferAttributes.attr_flags = BufferAttributeFlags;
+
+					FMemory::Memset(&BufferConfiguration.DataMemory, 0, sizeof(BufferConfiguration.DataMemory));
+					BufferConfiguration.DataMemory.max_size = BufferConfiguration.DataMemory.min_size = BufferConfiguration.PayloadExpectedSize;
+					BufferAttributes.data = &BufferConfiguration.DataMemory;
+
+					FMemory::Memset(&BufferConfiguration.HeaderMemory, 0, sizeof(BufferConfiguration.HeaderMemory));
+					BufferConfiguration.HeaderMemory.max_size = BufferConfiguration.HeaderMemory.min_size = BufferConfiguration.HeaderExpectedSize;
+					BufferAttributes.hdr = &BufferConfiguration.HeaderMemory;
+
+					rmax_status_t Status = rmax_in_query_buffer_size(StreamType, &RivermaxInterface, &BufferAttributes, &BufferConfiguration.PayloadSize, &BufferConfiguration.HeaderSize);
+					if (Status == RMAX_OK)
+					{
+						const uint32 CacheLineSize = PLATFORM_CACHE_LINE_SIZE;
+
+						BufferConfiguration.DataMemory.ptr = FMemory::Malloc(BufferConfiguration.PayloadSize, CacheLineSize);
+						BufferConfiguration.HeaderMemory.ptr = FMemory::Malloc(BufferConfiguration.HeaderSize, CacheLineSize);
+
+						//Buffers configured, now configure stream and attach flow
+						const rmax_in_timestamp_format TimestampFormat = rmax_in_timestamp_format::RMAX_PACKET_TIMESTAMP_RAW_NANO; //how packets are stamped. counter or nanoseconds
+						const rmax_in_flags InputFlags = rmax_in_flags::RMAX_IN_CREATE_STREAM_INFO_PER_PACKET; //default value for 2110 in example
+						Status = rmax_in_create_stream(StreamType, &RivermaxInterface, &BufferAttributes, TimestampFormat, InputFlags, &StreamId);
+						if (Status == RMAX_OK)
+						{
+							Status = rmax_in_attach_flow(StreamId, &FlowAttribute);
+							if (Status == RMAX_OK)
+							{
+								bIsActive = true;
+								RivermaxThread = FRunnableThread::Create(this, TEXT("Rivermax InputStream Thread"), 128 * 1024, TPri_AboveNormal, FPlatformAffinity::GetPoolThreadMask());
+								bWasSuccessful = true;
+							}
+							else
+							{
+								UE_LOG(LogRivermax, Warning, TEXT("Could not attach flow to stream. Status: %d."), Status);
+							}
+						}
+						else
+						{
+							UE_LOG(LogRivermax, Warning, TEXT("Could not create stream. Status: %d."), Status);
+						}
+					}
+					else
+					{
+						UE_LOG(LogRivermax, Warning, TEXT("Could not query buffer size. Status: %d"), Status);
+					}
+				}
+			}
+
+			Listener->OnInitializationCompleted(bWasSuccessful);
+		});
+		
+		return true;
+	}
+
+	void FRivermaxInputStream::Uninitialize()
+	{
+		bIsShuttingDown = true;
+
+		//If init task is ongoing, wait till it's done
+		if (InitTaskFuture.IsReady() == false)
+		{
+			InitTaskFuture.Wait();
+		}
+
+		if (RivermaxThread != nullptr)
+		{
+			Stop();
+			RivermaxThread->Kill(true);
+			delete RivermaxThread;
+			RivermaxThread = nullptr;
+			UE_LOG(LogRivermax, Log, TEXT("Rivermax Input stream has shutdown"));
+		}
+	}
+
+	void FRivermaxInputStream::Process_AnyThread()
+	{
+		const size_t MinChunkSize = 0;
+		const size_t MaxChunkSize = 5000;
+		const int Timeout = 0;
+		const int Flags = 0;
+		rmax_in_completion Completion;
+		rmax_status_t Status = rmax_in_get_next_chunk(StreamId, MinChunkSize, MaxChunkSize, Timeout, Flags, &Completion);
+		if (Status == RMAX_OK)
+		{
+			ParseChunk(Completion);
+		}
+		else
+		{
+			UE_LOG(LogRivermax, Warning, TEXT("Rivermax Input stream failed to get next chunk. Status: %d"), Status);
+		}
+	}
+
+	bool FRivermaxInputStream::Init()
+	{
+		return true;
+	}
+
+	uint32 FRivermaxInputStream::Run()
+	{
+		while (bIsActive)
+		{
+			Process_AnyThread();
+		}
+
+		if (StreamId)
+		{
+			rmax_status_t Status = rmax_in_detach_flow(StreamId, &FlowAttribute);
+			if (Status != RMAX_OK)
+			{
+				UE_LOG(LogRivermax, Warning, TEXT("Failed to detach rivermax flow %d from input stream %d. Status: %d"), FlowAttribute.flow_id, StreamId, Status);
+			}
+
+			Status = rmax_in_destroy_stream(StreamId);
+
+			if (Status != RMAX_OK)
+			{
+				UE_LOG(LogRivermax, Warning, TEXT("Failed to destroy input stream %d correctly. Status: %d"), StreamId, Status);
+			}
+		}
+
+		return 0;
+	}
+
+	void FRivermaxInputStream::Stop()
+	{
+		bIsActive = false;
+	}
+
+	void FRivermaxInputStream::Exit()
+	{
+
+	}
+
+	bool FRivermaxInputStream::GetRTPParameter(uint8* InputRTP, FRTPParameter& OutParameter)
+	{
+		OutParameter.Timestamp = 0;
+		
+		const uint8* RTPHeaderPtr = GetRTPHeaderPointer(InputRTP);
+		if (((RTPHeaderPtr[0] & 0xC0) != 0x80))
+		{
+			return false;
+		}
+
+		OutParameter.SequencerNumber = RTPHeaderPtr[3] | RTPHeaderPtr[2] << 8;
+		if (IsExtendedSequenceNumber()) //todo if 2022 is supported extended sequence number is supported in other ways
+		{
+			OutParameter.SequencerNumber |= RTPHeaderPtr[12] << 24 | RTPHeaderPtr[13] << 16;
+			OutParameter.bIsFBit = !!(RTPHeaderPtr[16] & 0x80);
+		}
+		OutParameter.Timestamp = ByteSwap(*(uint32*)((RTPHeaderPtr)+4));
+		OutParameter.bIsMBit = !!(RTPHeaderPtr[1] & 0x80);
+		return true;
+	}
+
+	void FRivermaxInputStream::ParseChunk(const rmax_in_completion& Completion)
+	{
+		for (uint64 StrideIndex = 0; StrideIndex < Completion.chunk_size; ++StrideIndex)
+		{
+			ensure(Completion.hdr_ptr);
+			if (Completion.hdr_ptr == nullptr)
+			{
+				break;
+			}
+
+			uint8* HeaderPtr = (uint8*)Completion.hdr_ptr + StrideIndex * (size_t)BufferConfiguration.HeaderMemory.stride_size; //when using RMAX_RAW_PACKET then header is preceded by net header
+			uint8* DataPtr = (uint8*)Completion.data_ptr + StrideIndex * (size_t)BufferConfiguration.DataMemory.stride_size; // The payload is our data
+			uint32 DataOffset = 0;
+
+			if (Completion.packet_info_arr[StrideIndex].data_size)
+			{
+				FRTPParameter Parameter;
+				const bool bIsValid = GetRTPParameter(HeaderPtr, Parameter);
+				if (bIsValid)
+				{
+					if (bIsFirstFrameReceived)
+					{
+						//stats  rx count, received bytes++
+						uint64 LastSequenceNumberIncremented = StreamData.LastSequenceNumber + 1;
+						if (IsExtendedSequenceNumber() == false)
+						{
+							LastSequenceNumberIncremented = LastSequenceNumberIncremented & 0xFFFF;
+						}
+
+						bool bCanProcessChunk = true;
+						const uint64 LostPackets = ((uint64)Parameter.SequencerNumber + 0x100000000 - LastSequenceNumberIncremented) & 0xFFFFFFFF;
+						if (LostPackets > 0)
+						{
+							bCanProcessChunk = false;
+							StreamData.WritingOffset = 0;
+							StreamData.ReceivedSize = 0;
+							UE_LOG(LogRivermax, Warning, TEXT("Lost %uld packets"), LostPackets);
+						}
+
+						StreamData.LastSequenceNumber = Parameter.SequencerNumber;
+
+						//if flags are RMAX_IN_CREATE_STREAM_INFO_PER_PACKET todo 
+						{
+							if (FlowAttribute.flow_id && Completion.packet_info_arr[StrideIndex].flow_id != FlowAttribute.flow_id)
+							{
+								UE_LOG(LogRivermax, Error, TEXT("Received data from unexpected FlowId '%d'. Expected '%d'."), Completion.packet_info_arr[StrideIndex].flow_id, FlowAttribute.flow_id);
+							}
+						}
+
+						if (bCanProcessChunk)
+						{
+							TRACE_CPUPROFILER_EVENT_SCOPE(FRivermaxInputStream::ProcessingChunk)
+							const FRivermaxRTPSampleRowData* HeaderStart = reinterpret_cast<FRivermaxRTPSampleRowData*>(GetRTPHeaderPointer(HeaderPtr));
+
+							uint16 SizeToCopy = ByteSwap((uint16)HeaderStart->SRDLength1);
+							FMemory::Memcpy(&StreamData.CurrentFrame[StreamData.WritingOffset], &DataPtr[DataOffset], SizeToCopy);
+							StreamData.WritingOffset += SizeToCopy;
+							StreamData.ReceivedSize += SizeToCopy;
+
+							// todo Warning !!! GPUDirect doesn't support more than 1 SRD
+							if (HeaderStart->ContinuationBit1)
+							{
+								DataOffset += SizeToCopy;
+								SizeToCopy = ByteSwap((uint16)HeaderStart->SRDLength2);
+								FMemory::Memcpy(&StreamData.CurrentFrame[StreamData.WritingOffset], &DataPtr[DataOffset], SizeToCopy);
+								StreamData.WritingOffset += SizeToCopy;
+								StreamData.ReceivedSize += SizeToCopy;
+							}
+
+							if (StreamData.ReceivedSize > StreamData.ExpectedSize)
+							{
+								UE_LOG(LogRivermax, Warning, TEXT("Received too much data (%d). Expected %d but received (%d)"), StreamData.ReceivedSize - StreamData.ExpectedSize, StreamData.ExpectedSize, StreamData.ReceivedSize);
+								StreamData.WritingOffset = 0;
+								StreamData.ReceivedSize = 0;
+							}
+							else if (HeaderStart->MarkerBit)
+							{
+								if (StreamData.ReceivedSize == StreamData.ExpectedSize)
+								{
+									TRACE_CPUPROFILER_EVENT_SCOPE(FRivermaxInputStream::ProcessingReceivedFrame)
+									FRivermaxInputVideoFrameDescriptor Descriptor;
+									Descriptor.Width = Options.Resolution.X;
+									Descriptor.Height = Options.Resolution.Y;
+									Descriptor.Stride = Options.Resolution.X * 2;//todo stride. Bytes per row
+									FRivermaxInputVideoFrameReception NewFrame;
+									NewFrame.VideoBuffer = StreamData.CurrentFrame;
+									Listener->OnVideoFrameReceived(Descriptor, NewFrame);//todo stride
+									PrepareNextFrame();
+								}
+								else
+								{
+									UE_LOG(LogRivermax, Warning, TEXT("End of frame received (Marker bit) but not enough data was received (missing %d). Expected %d but received (%d)"), StreamData.ExpectedSize - StreamData.ReceivedSize, StreamData.ExpectedSize, StreamData.ReceivedSize);
+									StreamData.WritingOffset = 0;
+									StreamData.ReceivedSize = 0;
+								}
+							}
+						}
+
+					}
+					else
+					{
+						if (Parameter.bIsMBit)
+						{
+							StreamData.LastSequenceNumber = Parameter.SequencerNumber;
+							bIsFirstFrameReceived = true;
+							PrepareNextFrame();
+						}
+					}
+
+					//todo why  Reset RTP header
+					*(uint32_t*)GetRTPHeaderPointer(HeaderPtr) = 0;
+				}
+			}
+
+			//todo why?
+			Completion.packet_info_arr[StrideIndex].data_size = 0;
+		}
+	}
+
+	bool FRivermaxInputStream::IsExtendedSequenceNumber() const
+	{
+		return RivermaxStreamType == ERivermaxStreamType::VIDEO_2110_20_STREAM;
+	}
+
+	void FRivermaxInputStream::PrepareNextFrame()
+	{
+		using namespace UE::RivermaxCore::Private::Utils;
+
+		FRivermaxInputVideoFrameDescriptor Descriptor;
+		FRivermaxInputVideoFrameRequest Request;
+		const uint32 Groups = Options.Resolution.X / PixelsPerGroup_422_8b;
+		const uint32 BytesPerLine = Groups * BytesPerGroup_422_8b;
+		Descriptor.VideoBufferSize = Options.Resolution.Y * BytesPerLine;
+		Listener->OnVideoFrameRequested(Descriptor, Request);
+		StreamData.CurrentFrame = Request.VideoBuffer;
+		StreamData.WritingOffset = 0;
+		StreamData.ReceivedSize = 0;
+		StreamData.ExpectedSize = Descriptor.VideoBufferSize;
+	}
+}
+
+
+
+
+
diff --git a/Engine/Plugins/VirtualProduction/Rivermax/RivermaxCore/Source/RivermaxCore/Private/RivermaxInputStream.h b/Engine/Plugins/VirtualProduction/Rivermax/RivermaxCore/Source/RivermaxCore/Private/RivermaxInputStream.h
new file mode 100644
--- /dev/null
+++ b/Engine/Plugins/VirtualProduction/Rivermax/RivermaxCore/Source/RivermaxCore/Private/RivermaxInputStream.h
@@ -0,0 +1,99 @@
+// Copyright Epic Games, Inc. All Rights Reserved.
+
+#pragma once
+
+#include "IRivermaxInputStream.h"
+
+#include <atomic>
+#include "Async/Future.h"
+#include "HAL/Runnable.h"
+#include "RivermaxHeader.h"
+#include "RivermaxTypes.h"
+
+
+
+namespace UE::RivermaxCore::Private
+{
+	using UE::RivermaxCore::IRivermaxInputStream;
+	using UE::RivermaxCore::IRivermaxInputStreamListener;
+	using UE::RivermaxCore::FRivermaxStreamOptions;
+	
+	struct FInputStreamBufferConfiguration
+	{
+		size_t PayloadSize = 0;
+		size_t HeaderSize = 0;
+		uint16 PayloadExpectedSize = 1500;
+		uint16 HeaderExpectedSize = 20; //for 2110
+
+		rmax_in_memblock DataMemory;
+		rmax_in_memblock HeaderMemory;
+	};
+
+	struct FInputStreamData
+	{
+		uint64 LastSequenceNumber = 0;
+		uint8* CurrentFrame = nullptr;
+		uint32 WritingOffset = 0;
+		uint32 ReceivedSize = 0;
+		uint32 ExpectedSize = 0;
+	};
+
+	class FRivermaxInputStream : public IRivermaxInputStream, public FRunnable
+	{
+	public:
+		FRivermaxInputStream();
+		virtual ~FRivermaxInputStream();
+
+	public:
+
+		//~ Begin IRivermaxInputStreamListener interface
+		virtual bool Initialize(const FRivermaxStreamOptions& InOptions, IRivermaxInputStreamListener& InListener) override;
+		virtual void Uninitialize() override;
+		//~ End IRivermaxInputStreamListener interface 
+
+		void Process_AnyThread();
+
+		//~ Begin FRunnable interface
+		virtual bool Init() override;
+		virtual uint32 Run() override;
+		virtual void Stop() override;
+		virtual void Exit() override;
+		//~ End FRunnable interface
+
+	private:
+
+		struct FRTPParameter
+		{
+			uint32 SequencerNumber = 0;
+			uint32 Timestamp = 0;
+			bool bIsMBit = false;
+			bool bIsFBit = false;
+		};
+
+		bool GetRTPParameter(uint8* InputRTP, FRTPParameter& OutParameter);
+		void ParseChunk(const rmax_in_completion& Completion);
+		bool IsExtendedSequenceNumber() const;
+		void PrepareNextFrame();
+
+	private:
+
+
+
+		FRivermaxStreamOptions Options;
+
+		FRunnableThread* RivermaxThread = nullptr;
+		std::atomic<bool> bIsActive;
+
+		rmax_stream_id StreamId = 0;
+		rmax_in_flow_attr FlowAttribute;
+		FInputStreamBufferConfiguration BufferConfiguration;
+		bool bIsFirstFrameReceived = false;
+		ERivermaxStreamType RivermaxStreamType = ERivermaxStreamType::VIDEO_2110_20_STREAM;
+		FInputStreamData StreamData;
+		IRivermaxInputStreamListener* Listener;
+		std::atomic<bool> bIsShuttingDown;
+		TFuture<void> InitTaskFuture;
+	};
+}
+
+
diff --git a/Engine/Plugins/VirtualProduction/Rivermax/RivermaxCore/Source/RivermaxCore/Private/RivermaxLog.cpp b/Engine/Plugins/VirtualProduction/Rivermax/RivermaxCore/Source/RivermaxCore/Private/RivermaxLog.cpp
new file mode 100644
--- /dev/null
+++ b/Engine/Plugins/VirtualProduction/Rivermax/RivermaxCore/Source/RivermaxCore/Private/RivermaxLog.cpp
@@ -0,0 +1,5 @@
+// Copyright Epic Games, Inc. All Rights Reserved.
+
+#include "RivermaxLog.h"
+
+DEFINE_LOG_CATEGORY(LogRivermax);
diff --git a/Engine/Plugins/VirtualProduction/Rivermax/RivermaxCore/Source/RivermaxCore/Private/RivermaxLog.h b/Engine/Plugins/VirtualProduction/Rivermax/RivermaxCore/Source/RivermaxCore/Private/RivermaxLog.h
new file mode 100644
--- /dev/null
+++ b/Engine/Plugins/VirtualProduction/Rivermax/RivermaxCore/Source/RivermaxCore/Private/RivermaxLog.h
@@ -0,0 +1,7 @@
+// Copyright Epic Games, Inc. All Rights Reserved.
+
+#pragma once
+
+#include "CoreMinimal.h"
+
+DECLARE_LOG_CATEGORY_EXTERN(LogRivermax, Log, All);
diff --git a/Engine/Plugins/VirtualProduction/Rivermax/RivermaxCore/Source/RivermaxCore/Private/RivermaxManager.cpp b/Engine/Plugins/VirtualProduction/Rivermax/RivermaxCore/Source/RivermaxCore/Private/RivermaxManager.cpp
new file mode 100644
--- /dev/null
+++ b/Engine/Plugins/VirtualProduction/Rivermax/RivermaxCore/Source/RivermaxCore/Private/RivermaxManager.cpp
@@ -0,0 +1,112 @@
+// Copyright Epic Games, Inc. All Rights Reserved.
+
+#include "RivermaxManager.h"
+
+#include "Misc/App.h"
+#include "Misc/Paths.h"
+#include "Misc/CoreDelegates.h"
+#include "RivermaxHeader.h"
+#include "RivermaxLog.h"
+
+
+namespace UE::RivermaxCore::Private
+{
+	FRivermaxManager::FRivermaxManager()
+	{
+		const auto InitializeRivermaxLibFunc = [this]()
+		{
+			rmax_init_config Config;
+			memset(&Config, 0, sizeof(Config));
+			//Config.flags |= RIVERMAX_HANDLE_SIGNAL; //Verify if desired
+
+			rmax_status_t Status = rmax_init(&Config);
+			if (Status == RMAX_OK)
+			{
+				uint32 Major = 0;
+				uint32 Minor = 0;
+				uint32 Release = 0;
+				uint32 Build = 0;
+				Status = rmax_get_version(&Major, &Minor, &Release, &Build);
+				if (Status == RMAX_OK)
+				{
+					bIsInitialized = true;
+					UE_LOG(LogRivermax, Log, TEXT("Rivermax library version %d.%d.%d.%d succesfully initialized"), Major, Minor, Release, Build);
+				}
+				else
+				{
+					UE_LOG(LogRivermax, Log, TEXT("Failed to retrieve Rivermax library version. Status: %d"), Status);
+				}
+			}
+			else if (Status == RMAX_ERR_LICENSE_ISSUE)
+			{
+				UE_LOG(LogRivermax, Log, TEXT("Rivermax License could not be found. Have you configured RIVERMAX_LICENSE_PATH environment variable?"));
+			}
+			
+			if(bIsInitialized == false)
+			{
+				UE_LOG(LogRivermax, Log, TEXT("Rivermax library failed to initialize. Status: %d"), Status);
+			}
+		};
+
+		const bool bIsLibraryLoaded = LoadRivermaxLibrary();
+		if (bIsLibraryLoaded && FApp::CanEverRender())
+		{
+			//Postpone initialization after all modules have been loaded
+			FCoreDelegates::OnAllModuleLoadingPhasesComplete.AddLambda(InitializeRivermaxLibFunc);
+		}
+		else
+		{
+			UE_LOG(LogRivermax, Log, TEXT("Skipping Rivermax initialization. Library won't be usable."));
+		}
+	}
+
+	FRivermaxManager::~FRivermaxManager()
+	{
+		if (IsInitialized())
+		{
+			rmax_status_t Status = rmax_cleanup();
+			if (Status != RMAX_OK)
+			{
+				UE_LOG(LogRivermax, Log, TEXT("Failed to cleanup Rivermax Library. Status: %d"), Status);
+			}
+			UE_LOG(LogRivermax, Log, TEXT("Rivermax Library has shutdown."));
+		}
+
+		if (LibraryHandle != nullptr)
+		{
+			FPlatformProcess::FreeDllHandle(LibraryHandle);
+			LibraryHandle = nullptr;
+		}
+	}
+
+	bool FRivermaxManager::IsInitialized() const
+	{
+		return bIsInitialized;
+	}
+
+	bool FRivermaxManager::LoadRivermaxLibrary()
+	{
+#if defined(RIVERMAX_LIBRARY_PLATFORM_PATH) && defined(RIVERMAX_LIBRARY_NAME)
+		const FString LibraryPath = FPaths::Combine(FPaths::EngineSourceDir(), TEXT("ThirdParty/NVIDIA/Rivermax/lib"), TEXT(PREPROCESSOR_TO_STRING(RIVERMAX_LIBRARY_PLATFORM_PATH)));
+		const FString LibraryName = TEXT(PREPROCESSOR_TO_STRING(RIVERMAX_LIBRARY_NAME));
+
+		FPlatformProcess::PushDllDirectory(*LibraryPath);
+		LibraryHandle = FPlatformProcess::GetDllHandle(*LibraryName);
+		FPlatformProcess::PopDllDirectory(*LibraryPath);
+#endif
+
+		const bool bIsLibraryValid = LibraryHandle != nullptr;
+		if(!bIsLibraryValid)
+		{
+			UE_LOG(LogRivermax, Log, TEXT("Failed to load required library %s. Rivermax library will not be functional."), *LibraryName);
+		}
+
+		return bIsLibraryValid;
+	}
+
+}
+
+
+
+
+
diff --git a/Engine/Plugins/VirtualProduction/Rivermax/RivermaxCore/Source/RivermaxCore/Private/RivermaxManager.h b/Engine/Plugins/VirtualProduction/Rivermax/RivermaxCore/Source/RivermaxCore/Private/RivermaxManager.h
new file mode 100644
--- /dev/null
+++ b/Engine/Plugins/VirtualProduction/Rivermax/RivermaxCore/Source/RivermaxCore/Private/RivermaxManager.h
@@ -0,0 +1,30 @@
+// Copyright Epic Games, Inc. All Rights Reserved.
+
+#pragma once
+
+#include "IRivermaxManager.h"
+
+
+namespace UE::RivermaxCore::Private
+{
+	class FRivermaxManager : public UE::RivermaxCore::IRivermaxManager
+	{
+	public:
+		FRivermaxManager();
+		~FRivermaxManager();
+
+	public:
+		//~ Begin IRivermaxManager interface
+		virtual bool IsInitialized() const override;
+		//~ End IRivermaxManager interface
+
+	private:
+		bool LoadRivermaxLibrary();
+
+	private:
+		bool bIsInitialized = false;
+		void* LibraryHandle = nullptr;
+	};
+}
+
+
diff --git a/Engine/Plugins/VirtualProduction/Rivermax/RivermaxCore/Source/RivermaxCore/Private/RivermaxOutputFrame.cpp b/Engine/Plugins/VirtualProduction/Rivermax/RivermaxCore/Source/RivermaxCore/Private/RivermaxOutputFrame.cpp
new file mode 100644
--- /dev/null
+++ b/Engine/Plugins/VirtualProduction/Rivermax/RivermaxCore/Source/RivermaxCore/Private/RivermaxOutputFrame.cpp
@@ -0,0 +1,36 @@
+// Copyright Epic Games, Inc. All Rights Reserved.
+
+#include "RivermaxOutputFrame.h"
+
+namespace UE::RivermaxCore::Private
+{
+	FRivermaxOutputFrame::FRivermaxOutputFrame(uint32 InFrameIndex)
+		: FrameIndex(InFrameIndex)
+		, VideoBuffer(nullptr)
+	{
+		Clear();
+	}
+
+	FRivermaxOutputFrame::~FRivermaxOutputFrame()
+	{
+		if (VideoBuffer)
+		{
+			FMemory::Free(VideoBuffer);
+		}
+	}
+
+	bool FRivermaxOutputFrame::IsReadyToBeSent() const
+	{
+		return FrameIdentifier != InvalidIdentifier && bIsVideoBufferReady;
+	}
+
+	void FRivermaxOutputFrame::Clear()
+	{
+		FrameIdentifier = InvalidIdentifier;
+		PacketCounter = 0;
+		LineNumber = 0;
+		SRDOffset = 0;
+		TimestampTicks = 0.0;
+		bIsVideoBufferReady = false;
+	}
+}
diff --git a/Engine/Plugins/VirtualProduction/Rivermax/RivermaxCore/Source/RivermaxCore/Private/RivermaxOutputFrame.h b/Engine/Plugins/VirtualProduction/Rivermax/RivermaxCore/Source/RivermaxCore/Private/RivermaxOutputFrame.h
new file mode 100644
--- /dev/null
+++ b/Engine/Plugins/VirtualProduction/Rivermax/RivermaxCore/Source/RivermaxCore/Private/RivermaxOutputFrame.h
@@ -0,0 +1,61 @@
+// Copyright Epic Games, Inc. All Rights Reserved.
+
+#pragma once
+
+#include "CoreMinimal.h"
+
+#include "Math/NumericLimits.h"
+#include "Misc/Timecode.h"
+
+namespace UE::RivermaxCore::Private
+{
+	/**
+	 * Output frame descriptor. Contains data to be sent on the wire and packetization tracking
+	 */
+	struct FRivermaxOutputFrame
+	{
+		FRivermaxOutputFrame(uint32 InFrameIndex);
+		~FRivermaxOutputFrame();
+		FRivermaxOutputFrame(const FRivermaxOutputFrame&) = delete;
+		FRivermaxOutputFrame& operator=(const FRivermaxOutputFrame&) = delete;
+
+		/** Whether this frame is ready to be sent */
+		bool IsReadyToBeSent() const;
+
+		/** Reset internal counters of this frame */
+		void Clear();
+
+		/** Index of the frame used as unique identifier */
+		uint32 FrameIndex = 0;
+
+		/** Video buffer where we copy texture to be sent out */
+		uint8* VideoBuffer = nullptr;
+
+		/** Identifier of the frame incremented every time a frame is captured */
+		static constexpr uint32 InvalidIdentifier = MAX_uint32;
+		uint32 FrameIdentifier = InvalidIdentifier;
+
+		/** Timecode at which this frame was captured */
+		FTimecode Timecode;
+
+		/** Used to flag this frame ready to be sent */
+		bool bIsVideoBufferReady = false;
+
+		/** Variables about progress of packetization of the frame */
+		uint32 PacketCounter = 0;
+		uint32 LineNumber = 0;
+		uint16 SRDOffset = 0;
+		uint32 ChunkNumber = 0;
+
+		/** Timestamp of this frame  */
+		double TimestampTicks = 0.0;
+
+		/** Pointers retrieved from Rivermax for the next chunk */
+		void* PayloadPtr = nullptr;
+		void* HeaderPtr = nullptr;
+
+		
+	};
+}
+
+
diff --git a/Engine/Plugins/VirtualProduction/Rivermax/RivermaxCore/Source/RivermaxCore/Private/RivermaxOutputStream.cpp b/Engine/Plugins/VirtualProduction/Rivermax/RivermaxCore/Source/RivermaxCore/Private/RivermaxOutputStream.cpp
new file mode 100644
--- /dev/null
+++ b/Engine/Plugins/VirtualProduction/Rivermax/RivermaxCore/Source/RivermaxCore/Private/RivermaxOutputStream.cpp
@@ -0,0 +1,682 @@
+// Copyright Epic Games, Inc. All Rights Reserved.
+
+#include "RivermaxOutputStream.h"
+
+#include "Async/Async.h"
+#include "IRivermaxCoreModule.h"
+#include "IRivermaxManager.h"
+#include "Misc/ByteSwap.h"
+#include "RivermaxLog.h"
+#include "RivermaxTypes.h"
+#include "RivermaxUtils.h"
+
+
+namespace UE::RivermaxCore::Private
+{
+	void CalculateStreamTime(const FRivermaxOutputStreamMemory& InStreamMemory, const FRivermaxOutputStreamData& InStreamData, const FRivermaxStreamOptions& InMediaOptions, const double InSampleRate, double& InOutTimeNs, double& OutTimestampTick)
+	{
+		using namespace UE::RivermaxCore::Private::Utils;
+
+		double FrameIntervalNs = InStreamData.FrameFieldTimeIntervalNs;
+		const bool bIsProgressive = true;//todo MediaConfiguration.IsProgressive() 
+		const ERivermaxStreamType StreamType = ERivermaxStreamType::VIDEO_2110_20_STREAM; //todo
+		if (bIsProgressive == false)
+		{
+			FrameIntervalNs *= 2;
+		}
+		const uint64 FrameAtTime = (uint64)(InOutTimeNs / FrameIntervalNs + 1);
+		double FirstPacketStartTimeNs = FrameAtTime * FrameIntervalNs; //next alignment point calculation
+
+		uint32 PacketsInFrameField = InStreamMemory.PacketsInFrameField;
+		if (bIsProgressive == false)
+		{
+			PacketsInFrameField *= 2;
+		}
+		
+		// See https://ieeexplore.ieee.org/stamp/stamp.jsp?tp=&arnumber=8165971 for reference
+		// To be revisited for non standard resolution/frame rate
+		if (StreamType == ERivermaxStreamType::VIDEO_2110_20_STREAM)
+		{
+			double RActive;
+			double TRODefaultMultiplier;
+			if (bIsProgressive)
+			{
+				RActive = (1080.0 / 1125.0);
+				if (InMediaOptions.Resolution.Y >= FullHDHeight)
+				{
+				 	// As defined by SMPTE 2110-21 6.3.2
+					TRODefaultMultiplier = (43.0 / 1125.0);
+				}
+				else 
+				{
+					TRODefaultMultiplier = (28.0 / 750.0);
+				}
+			}
+			else
+			{
+				if (InMediaOptions.Resolution.Y >= FullHDHeight)
+				{
+					// As defined by SMPTE 2110-21 6.3.3
+					RActive = (1080.0 / 1125.0);
+					TRODefaultMultiplier = (22.0 / 1125.0);
+				}
+				else if (InMediaOptions.Resolution.Y >= 576)
+				{
+					RActive = (576.0 / 625.0);
+					TRODefaultMultiplier = (26.0 / 625.0);
+				}
+				else 
+				{
+					RActive = (487.0 / 525.0);
+					TRODefaultMultiplier = (20.0 / 525.0);
+				}
+			}
+
+			const double TRSNano = (FrameIntervalNs * RActive) / PacketsInFrameField;
+			const double TROffset = (TRODefaultMultiplier * FrameIntervalNs) - (VideoTROModification * TRSNano);
+			FirstPacketStartTimeNs += TROffset;
+		}
+
+		constexpr double OneSecondNs = 1 * 1E9;
+		OutTimestampTick = FirstPacketStartTimeNs * (InSampleRate / OneSecondNs);
+		InOutTimeNs = FirstPacketStartTimeNs;
+	}
+	
+	uint32 FindPayloadSize(const FRivermaxStreamOptions& InOptions, uint32 InBytesPerLine)
+	{
+		using namespace UE::RivermaxCore::Private::Utils;
+
+		// If stride per line is small then our minimal payload size, send one line per packet.
+		if (InBytesPerLine <= MaxPayloadSize)
+		{
+			return InBytesPerLine;
+		}
+
+		uint32 PacketDivider = 2; //Start with 2 packets per line as 1 is too big already
+		static constexpr uint32 MaxPacketCountPerLine = 50; 
+		while (PacketDivider <= MaxPacketCountPerLine)
+		{
+			const uint32 TestPayloadSize = InBytesPerLine / PacketDivider;
+			const bool bIsFlushInLine = InBytesPerLine % PacketDivider == 0;
+			if (bIsFlushInLine && TestPayloadSize < MaxPayloadSize)
+			{
+				return TestPayloadSize;
+			}
+
+			++PacketDivider;
+		}
+
+		return 0;
+	}
+
+	uint32 FindChunksPerLine(const FRivermaxStreamOptions& InOptions)
+	{
+		uint32 ChunksPerLine = 1; //Will need to revisit the impact of that
+		static constexpr uint32 MaxChunksPerLine = 10;
+		while (ChunksPerLine <= MaxChunksPerLine)
+		{
+			if (InOptions.Resolution.Y % ChunksPerLine == 0)
+			{
+				return ChunksPerLine;
+			}
+
+			++ChunksPerLine;
+		}
+
+		return 0;
+	}
+
+
+	FRivermaxOutputStream::FRivermaxOutputStream()
+		: bIsActive(false)
+	{
+
+	}
+
+	FRivermaxOutputStream::~FRivermaxOutputStream()
+	{
+		Uninitialize();
+	}
+
+	bool FRivermaxOutputStream::Initialize(const FRivermaxStreamOptions& InOptions, IRivermaxOutputStreamListener& InListener)
+	{
+		IRivermaxCoreModule& RivermaxModule = FModuleManager::LoadModuleChecked<IRivermaxCoreModule>(TEXT("RivermaxCore"));
+		if (RivermaxModule.GetRivermaxManager()->IsInitialized() == false)
+		{
+			UE_LOG(LogRivermax, Warning, TEXT("Can't create Rivermax Output Stream. Library isn't initialized."));
+			return false;
+		}
+
+		Options = InOptions;
+		Listener = &InListener;
+
+		Async(EAsyncExecution::TaskGraph, [this]()
+		{
+			//Create event to wait on when no frames are available to send
+			ReadyToSendEvent = FPlatformProcess::GetSynchEventFromPool();
+
+			TAnsiStringBuilder<2048> SDPDescription;
+			UE::RivermaxCore::Private::Utils::StreamOptionsToSDPDescription(Options, SDPDescription);
+
+			// Initialize video buffer memory we will be using
+			InitializeBuffers();
+
+			// Initialize internal memory and rivermax configuration 
+			InitializeMemory();
+
+			// Initialize rivermax output stream with desired config
+			const uint32 NumberPacketsPerFrame = StreamMemory.PacketsInFrameField;
+			const uint32 MediaBlockIndex = 0;
+			rmax_stream_id NewId;
+			rmax_qos_attr QOSAttributes = { 0, 0 }; //todo
+
+			rmax_buffer_attr BufferAttributes;
+			BufferAttributes.chunk_size_in_strides = StreamMemory.ChunkSizeInStrides;
+			BufferAttributes.data_stride_size = StreamMemory.PayloadSize; //Stride between chunks. 
+			BufferAttributes.app_hdr_stride_size = StreamMemory.HeaderStrideSize; 
+			BufferAttributes.mem_block_array = StreamMemory.MemoryBlocks.GetData();
+			BufferAttributes.mem_block_array_len = StreamMemory.MemoryBlocks.Num();
+
+			rmax_status_t Status = rmax_out_create_stream(SDPDescription.GetData(), &BufferAttributes, &QOSAttributes, NumberPacketsPerFrame, MediaBlockIndex, &NewId);
+			if (Status == RMAX_OK)
+			{
+				struct sockaddr_in SourceAddress;
+				struct sockaddr_in DestinationAddress;
+				Status = rmax_out_query_address(NewId, MediaBlockIndex, &SourceAddress, &DestinationAddress);
+				if (Status == RMAX_OK)
+				{
+					StreamId = NewId;
+
+					// Todo add event manager to avoid spinning
+
+					StreamData.FrameFieldTimeIntervalNs = 1E9 / Options.FrameRate.AsDecimal();
+					StreamData.SendTimeNs = (FPlatformTime::Seconds() + 1) * 1E9; //Current time + 1 sec todo why?
+					CalculateStreamTime(StreamMemory, StreamData, Options, MediaClockSampleRate, StreamData.SendTimeNs, StreamData.InitialTimestampTick);
+					StreamData.StartSendTimeNs = StreamData.SendTimeNs;
+
+					bIsActive = true;
+					RivermaxThread.Reset(FRunnableThread::Create(this, TEXT("Rmax OutputStream Thread"), 128 * 1024, TPri_AboveNormal, FPlatformAffinity::GetPoolThreadMask()));
+				}
+			}
+			else
+			{
+				UE_LOG(LogRivermax, Warning, TEXT("Failed to create Rivermax output stream. Status: %d"), Status);
+			}
+
+			Listener->OnInitializationCompleted(bIsActive);
+		});
+
+		return true;
+	}
+
+	void FRivermaxOutputStream::Uninitialize()
+	{
+		if (RivermaxThread != nullptr)
+		{
+			Stop();
+			
+			ReadyToSendEvent->Trigger();
+			RivermaxThread->Kill(true);
+			
+			RivermaxThread.Reset();
+
+
+			FPlatformProcess::ReturnSynchEventToPool(ReadyToSendEvent);
+			ReadyToSendEvent = nullptr;
+			UE_LOG(LogRivermax, Log, TEXT("Rivermax Output stream has shutdown"));
+		}
+
+	}
+
+	bool FRivermaxOutputStream::PushVideoFrame(const FRivermaxOutputVideoFrameInfo& NewFrame)
+	{
+		if (TSharedPtr<FRivermaxOutputFrame> AvailableFrame = GetNextAvailableFrame(NewFrame.FrameIdentifier))
+		{
+			FMemory::Memcpy(AvailableFrame->VideoBuffer, NewFrame.VideoBuffer, NewFrame.Height * Options.Stride);
+			AvailableFrame->bIsVideoBufferReady = true;
+
+			//If Frame ready to be sent
+			if(AvailableFrame->IsReadyToBeSent())
+			{
+				const FString TraceName = FString::Format(TEXT("FRivermaxOutputStream::PushFrame {0}"), { AvailableFrame->FrameIndex });
+				TRACE_CPUPROFILER_EVENT_SCOPE_TEXT(*TraceName);
+
+				FScopeLock Lock(&FrameCriticalSection);
+				
+				AvailableFrames.Remove(AvailableFrame);
+				FramesToSend.Add(MoveTemp(AvailableFrame));
+				ReadyToSendEvent->Trigger();
+			}
+
+			return true;
+		}
+		
+		return false;
+	}
+
+	void FRivermaxOutputStream::Process_AnyThread()
+	{
+		if (CurrentFrame.IsValid())
+		{
+			{
+				TRACE_CPUPROFILER_EVENT_SCOPE(FRivermaxOutputStream::Wait);
+				StreamData.SendTimeNs = StreamData.StartSendTimeNs + StreamData.FrameFieldTimeIntervalNs * StreamMemory.FramesFieldPerMemoryBlock * Stats.MemoryBlockSentCounter;
+				WaitForNextRound(StreamData.SendTimeNs);
+			}
+
+			const FString TraceName = FString::Format(TEXT("FRivermaxOutputStream::SendFrame {0}"), { CurrentFrame->FrameIndex });
+			TRACE_CPUPROFILER_EVENT_SCOPE_TEXT(*TraceName);
+
+			do 
+			{
+				if (bIsActive)
+				{
+					GetNextChunk();
+				}
+
+				if (bIsActive)
+				{
+					SetupRTPHeaders();
+				}
+
+				if (bIsActive)
+				{
+					CommitNextChunks();
+				}
+
+				//Update frame progress
+				if (bIsActive)
+				{
+					Stats.TotalStrides += StreamMemory.ChunkSizeInStrides;
+
+					if ((++CurrentFrame->ChunkNumber) % StreamMemory.ChunksPerFrameField == 0)
+					{
+						StreamData.SendTimeNs += StreamData.FrameFieldTimeIntervalNs;
+					}
+				}
+
+			} while (CurrentFrame->ChunkNumber < StreamMemory.ChunksPerMemoryBlock && bIsActive);
+			
+			//We are done with this frame, put it back as available
+			CurrentFrame->Clear();
+			{
+				FScopeLock Lock(&FrameCriticalSection);
+				AvailableFrames.Add(MoveTemp(CurrentFrame));
+				CurrentFrame.Reset();
+			}
+
+			Stats.MemoryBlockSentCounter++;
+			Stats.TotalStrides += StreamMemory.ChunkSizeInStrides;
+			StreamData.bHasFrameFirstChunkBeenFetched = false;
+		}
+	}
+
+	void FRivermaxOutputStream::InitializeBuffers()
+	{
+		AvailableFrames.Reserve(Options.NumberOfBuffers);
+		for (int32 Index = 0; Index < Options.NumberOfBuffers; ++Index)
+		{
+			TSharedPtr<FRivermaxOutputFrame> Frame = MakeShared<FRivermaxOutputFrame>(Index);
+			Frame->VideoBuffer = static_cast<uint8*>(FMemory::Malloc(Options.Resolution.Y * Options.Stride));
+			AvailableFrames.Add(MoveTemp(Frame));
+		}
+	}
+
+	void FRivermaxOutputStream::InitializeMemory()
+	{
+		using namespace UE::RivermaxCore::Private::Utils;
+
+		//2110 stream type
+		//We need to use the fullframe allocated size to compute the payload size.
+
+		const uint32 EffectiveBytesPerLine = Options.Stride;
+		const uint32 BytesPerLine = Options.Stride;
+		StreamMemory.PayloadSize = FindPayloadSize(Options, BytesPerLine); 
+		StreamMemory.DataStrideSize = BytesPerLine;//todo align cache line?
+		StreamMemory.LinesInChunk = FindChunksPerLine(Options);
+
+		StreamMemory.PacketsInLine =  FMath::RoundToInt32(BytesPerLine / (float)StreamMemory.PayloadSize);
+		StreamMemory.ChunkSizeInStrides = StreamMemory.LinesInChunk * StreamMemory.PacketsInLine;
+
+		StreamMemory.FramesFieldPerMemoryBlock = 1;
+		StreamMemory.PacketsInFrameField =  StreamMemory.PacketsInLine * Options.Resolution.Y;
+		StreamMemory.PacketsPerMemoryBlock = StreamMemory.PacketsInFrameField * StreamMemory.FramesFieldPerMemoryBlock;
+		StreamMemory.ChunksPerFrameField = StreamMemory.PacketsInFrameField / StreamMemory.ChunkSizeInStrides;
+		StreamMemory.ChunksPerMemoryBlock = StreamMemory.FramesFieldPerMemoryBlock * StreamMemory.ChunksPerFrameField;
+		StreamMemory.MemoryBlockCount = Options.NumberOfBuffers;
+		StreamMemory.StridesPerMemoryBlock = StreamMemory.ChunkSizeInStrides * StreamMemory.ChunksPerMemoryBlock;
+
+		// Setup arrays with the right sizes so we can give pointers to rivermax
+		StreamMemory.RTPHeaders.SetNumZeroed(StreamMemory.PacketsPerMemoryBlock);
+		StreamMemory.PayloadSizes.SetNumUninitialized(StreamMemory.PacketsPerMemoryBlock);
+		StreamMemory.HeaderSizes.SetNumUninitialized(StreamMemory.PacketsPerMemoryBlock);
+		StreamMemory.HeaderStrideSize = BytesPerHeader;
+		for (int32 PayloadSizeIndex = 0; PayloadSizeIndex < StreamMemory.PayloadSizes.Num(); ++PayloadSizeIndex)
+		{
+			//Go through each chunk to have effective payload size to be sent (last one of each line could be smaller)
+			if ((PayloadSizeIndex + 1) % StreamMemory.PacketsInLine == 0)
+			{
+				StreamMemory.PayloadSizes[PayloadSizeIndex] = EffectiveBytesPerLine - ((StreamMemory.PacketsInLine - 1) * StreamMemory.PayloadSize);
+			}
+			else
+			{
+				StreamMemory.PayloadSizes[PayloadSizeIndex] = StreamMemory.PayloadSize;
+			}
+			StreamMemory.HeaderSizes[PayloadSizeIndex] = StreamMemory.HeaderStrideSize;
+		}
+		StreamMemory.MemoryBlocks.SetNum(StreamMemory.MemoryBlockCount);
+		for (uint32 BlockIndex = 0; BlockIndex < StreamMemory.MemoryBlockCount; ++BlockIndex)
+		{
+			rmax_mem_block& Block = StreamMemory.MemoryBlocks[BlockIndex];
+			Block.chunks_num = StreamMemory.ChunksPerMemoryBlock;
+			Block.app_hdr_size_arr = StreamMemory.HeaderSizes.GetData();
+			Block.data_size_arr = StreamMemory.PayloadSizes.GetData();
+			Block.data_ptr = AvailableFrames[BlockIndex]->VideoBuffer;
+			Block.app_hdr_ptr = &StreamMemory.RTPHeaders[BlockIndex];
+		}
+	}
+
+	void FRivermaxOutputStream::InitializeNextFrame(const TSharedPtr<FRivermaxOutputFrame>& NextFrame)
+	{
+		NextFrame->TimestampTicks = StreamData.InitialTimestampTick;
+		NextFrame->LineNumber = 0;
+		NextFrame->PacketCounter = 0;
+		NextFrame->SRDOffset = 0;
+		NextFrame->ChunkNumber = 0;
+		NextFrame->PayloadPtr = nullptr;
+		NextFrame->HeaderPtr = nullptr;
+	}
+
+	TSharedPtr<FRivermaxOutputFrame> FRivermaxOutputStream::GetNextFrameToSend()
+	{
+		TSharedPtr<FRivermaxOutputFrame> FrameToSend;
+
+		FScopeLock Lock(&FrameCriticalSection);
+		if (FramesToSend.IsEmpty() == false)
+		{
+			// Pick oldest frame of the lot. 
+			// Todo Sort array when modified to always pick first item.
+			uint32 OldestIdentifier = TNumericLimits<uint32>::Max();
+			TArray<TSharedPtr<FRivermaxOutputFrame>>::TIterator It = FramesToSend.CreateIterator();
+			for (It; It; ++It)
+			{
+				TSharedPtr<FRivermaxOutputFrame>& ItFrame = *It;
+				if (ItFrame->FrameIdentifier < OldestIdentifier)
+				{
+					FrameToSend = *It;
+					OldestIdentifier = FrameToSend->FrameIdentifier;
+				}
+			}
+
+			if (ensure(FrameToSend))
+			{
+				FramesToSend.Remove(FrameToSend);
+				InitializeNextFrame(FrameToSend);
+			}
+		}
+
+		return FrameToSend;
+	}
+
+	TSharedPtr<FRivermaxOutputFrame> FRivermaxOutputStream::GetNextAvailableFrame(uint32 InFrameIdentifier)
+	{
+		TSharedPtr<FRivermaxOutputFrame> NextFrame;
+
+		{
+			FScopeLock Lock(&FrameCriticalSection);
+
+			//Find matching frame identifier
+			if(TSharedPtr<FRivermaxOutputFrame>* MatchingFrame = AvailableFrames.FindByPredicate([InFrameIdentifier](const TSharedPtr<FRivermaxOutputFrame>& Other){ return Other->FrameIdentifier == InFrameIdentifier; }))
+			{
+				NextFrame = *MatchingFrame;
+			}
+			else if(TSharedPtr<FRivermaxOutputFrame>* EmptyFrame = AvailableFrames.FindByPredicate([](const TSharedPtr<FRivermaxOutputFrame>& Other){ return Other->FrameIdentifier == FRivermaxOutputFrame::InvalidIdentifier; })) //Find an empty frame
+			{
+				NextFrame = *EmptyFrame;
+				NextFrame->FrameIdentifier = InFrameIdentifier;
+			}
+			else
+			{
+				// Could reuse a frame not yet ready to be sent
+				// Could also look into stealing a frame ready to be sent bot not sent yet
+			}
+		}
+
+		return NextFrame;
+	}
+
+	void FRivermaxOutputStream::BuildRTPHeader(FRTPHeader& OutHeader) const
+	{
+		// build RTP header - 12 bytes
+		/*
+		0                   1                   2                   3
+		0 1 2 3 4 5 6 7 8 9 0 1 2 3 4 5 6 7 8 9 0 1 2 3 4 5 6 7 8 9 0 1
+		+-+-+-+-+-+-+-+-+-+-+-+-+-+-+-+-+-+-+-+-+-+-+-+-+-+-+-+-+-+-+-+-+
+		| V |P|X|  CC   |M|     PT      |            SEQ                |
+		+-+-+-+-+-+-+-+-+-+-+-+-+-+-+-+-+-+-+-+-+-+-+-+-+-+-+-+-+-+-+-+-+
+		|                           timestamp                           |
+		+-+-+-+-+-+-+-+-+-+-+-+-+-+-+-+-+-+-+-+-+-+-+-+-+-+-+-+-+-+-+-+-+
+		|                           ssrc                                |
+		+-+-+-+-+-+-+-+-+-+-+-+-+-+-+-+-+-+-+-+-+-+-+-+-+-+-+-+-+-+-+-+-+*/
+		OutHeader.RawHeader[0] = 0x80;  // 10000000 - version2, no padding, no extension
+		OutHeader.RawHeader[1] = 96; //todo payload type from sdp file
+		OutHeader.RawHeader[2] = (StreamData.SequenceNumber >> 8) & 0xFF;  // sequence number
+		OutHeader.RawHeader[3] = (StreamData.SequenceNumber) & 0xFF;  // sequence number
+		*(uint32*)&OutHeader.RawHeader[4] = ByteSwap((uint32)CurrentFrame->TimestampTicks);
+
+		//2110 specific header
+		*(uint32*)&OutHeader.RawHeader[8] = 0x0eb51dbd;  // simulated ssrc
+
+		if (StreamType == ERivermaxStreamType::VIDEO_2110_20_STREAM)
+		{
+			//SRD means Sample Row Data
+			// build SRD header - 8-14 bytes
+		   /* 0                   1                   2                   3
+			0 1 2 3 4 5 6 7 8 9 0 1 2 3 4 5 6 7 8 9 0 1 2 3 4 5 6 7 8 9 0 1
+			+-+-+-+-+-+-+-+-+-+-+-+-+-+-+-+-+-+-+-+-+-+-+-+-+-+-+-+-+-+-+-+-+
+			|    Extended Sequence Number   |           SRD Length          |
+			+-+-+-+-+-+-+-+-+-+-+-+-+-+-+-+-+-+-+-+-+-+-+-+-+-+-+-+-+-+-+-+-+
+			|F|     SRD Row Number          |C|         SRD Offset          |
+			+-+-+-+-+-+-+-+-+-+-+-+-+-+-+-+-+-+-+-+-+-+-+-+-+-+-+-+-+-+-+-+-+ */
+			OutHeader.RawHeader[12] = (StreamData.SequenceNumber >> 24) & 0xff;  // high 16 bit of seq Extended Sequence Number
+			OutHeader.RawHeader[13] = (StreamData.SequenceNumber >> 16) & 0xff;  // high 16 bit of seq Extended Sequence Number
+
+			const uint16 CurrentPayloadSize = StreamMemory.PayloadSizes[StreamData.SequenceNumber % StreamMemory.PacketsInFrameField];
+			*(uint16*)&OutHeader.RawHeader[14] = ByteSwap(CurrentPayloadSize);  // SRD Length
+		
+			uint16 number_of_rows = Options.Resolution.Y; //todo divide by 2 if interlaced
+			uint16 srd_row_number = (CurrentFrame->LineNumber % number_of_rows);
+			*(uint16*)&OutHeader.RawHeader[16] = ByteSwap(srd_row_number);
+			OutHeader.RawHeader[16] |= (0 << 7); //todo when fields are sent for interlace
+
+			// we never have continuation
+			*(uint16*)&OutHeader.RawHeader[18] = ByteSwap(CurrentFrame->SRDOffset);  // SRD Offset
+			uint16 group_size = (uint16)((StreamMemory.PayloadSize /*- 20*/) / 2.5); // todo: need to support pgroup definition per sampling format to fully be compliant
+			CurrentFrame->SRDOffset = (CurrentFrame->SRDOffset + group_size) % (group_size * StreamMemory.PacketsInLine);
+		}
+
+		if (++CurrentFrame->PacketCounter == StreamMemory.PacketsInFrameField)
+		{
+			OutHeader.RawHeader[1] |= 0x80; // last packet in frame (Marker)
+			// ST2210-20: the timestamp SHOULD be the same for each packet of the frame/field.
+			const double ticks = (MediaClockSampleRate / (Options.FrameRate.AsDecimal()));
+			//if (set.video_type != VIDEO_TYPE::PROGRESSIVE) {
+			//	send_data.m_second_field = !send_data.m_second_field;
+			//	ticks /= 2;
+			//}
+			CurrentFrame->TimestampTicks += ticks;
+		}
+	}
+
+	void FRivermaxOutputStream::DestroyStream()
+	{
+		rmax_status_t Status = rmax_out_cancel_unsent_chunks(StreamId);
+		if (Status != RMAX_OK)
+		{
+			UE_LOG(LogRivermax, Warning, TEXT("Could not cancel unsent chunks when destroying output stream. Status: %d"), Status);
+		}
+
+		do 
+		{
+			Status = rmax_out_destroy_stream(StreamId);
+			if (RMAX_ERR_BUSY == Status) 
+			{
+				FPlatformProcess::SleepNoStats(0.01f);
+			}
+		} while (Status == RMAX_ERR_BUSY);
+	}
+
+	void FRivermaxOutputStream::WaitForNextRound(double NextRoundTimeNs)
+	{
+		double NowTimeNs = FPlatformTime::Seconds() * 1E9;
+
+		//If we are already late, no need to wait
+		if (NextRoundTimeNs <= NowTimeNs)
+		{
+			return;
+		}
+
+		const double SleepTimeNs = NextRoundTimeNs - NowTimeNs;
+
+		static const bool bShouldPoll = false;
+
+		//Todo Should we have a threshold where we don't sleep?
+		//Todo Should we sleep less when we had retries
+		if (bShouldPoll)
+		{
+			while (NowTimeNs < SleepTimeNs)
+			{
+				NowTimeNs = FPlatformTime::Seconds() * 1E9;
+			}
+		}
+		else
+		{
+			const double SleepTimeS = SleepTimeNs / 1E9;
+			UE_LOG(LogRivermax, Verbose, TEXT("Outputstream sleeping for %0.5f"), SleepTimeS);
+			FPlatformProcess::SleepNoStats(SleepTimeS);
+		}
+	}
+
+	void FRivermaxOutputStream::GetNextChunk()
+	{
+		rmax_status_t Status;
+		do
+		{
+			Status = rmax_out_get_next_chunk(StreamId, &CurrentFrame->PayloadPtr, &CurrentFrame->HeaderPtr);
+			if (Status == RMAX_OK)
+			{
+				if (StreamData.bHasFrameFirstChunkBeenFetched == false)
+				{
+					StreamData.bHasFrameFirstChunkBeenFetched = true;
+					if (CurrentFrame->VideoBuffer != CurrentFrame->PayloadPtr)
+					{
+						//Debug code to track rivermax chunk processing
+						UE_LOG(LogRivermax, Warning, TEXT("Frame being sent (%d) doesn't match chunks being processed."), CurrentFrame->FrameIndex);
+					}
+				}
+
+				break;
+			}
+			else if (Status == RMAX_ERR_NO_FREE_CHUNK)
+			{
+				//We should not be here
+				Stats.ChunkRetries++;
+			}
+			else
+			{
+				UE_LOG(LogRivermax, Error, TEXT("Invalid error happened while trying to get next chunks. Status: %d"), Status);
+				Listener->OnStreamError();
+				bIsActive = false;
+			}
+		} while (Status != RMAX_OK && bIsActive);
+	}
+
+	void FRivermaxOutputStream::SetupRTPHeaders()
+	{
+		uint8* HeaderRawPtr = reinterpret_cast<uint8*>(CurrentFrame->HeaderPtr);
+		check(HeaderRawPtr);
+		for (uint32 StrideIndex = 0; StrideIndex < StreamMemory.ChunkSizeInStrides && CurrentFrame->PacketCounter < StreamMemory.PacketsInFrameField; ++StrideIndex)
+		{
+			uint8* NextHeaderRawPtr = HeaderRawPtr + (StrideIndex * StreamMemory.HeaderStrideSize);
+			BuildRTPHeader(*reinterpret_cast<FRTPHeader*>(NextHeaderRawPtr));
+			//todo only for video
+			if (!((StrideIndex + 1) % StreamMemory.PacketsInLine))
+			{
+				CurrentFrame->LineNumber = (CurrentFrame->LineNumber + 1) % Options.Resolution.Y; //preparing line number for next iteration
+			}
+			++StreamData.SequenceNumber;
+		}
+	}
+
+	void FRivermaxOutputStream::CommitNextChunks()
+	{
+		rmax_status_t Status;
+		do
+		{
+			// todo configure timeout
+			
+			const uint64 Timeout = 0;
+			const rmax_commit_flags_t CommitFlags{};
+			Status = rmax_out_commit(StreamId, Timeout, CommitFlags);
+			if (Status == RMAX_OK)
+			{
+				break;
+			}
+			else if (Status == RMAX_ERR_HW_SEND_QUEUE_FULL)
+			{
+				Stats.CommitRetries++;
+			}
+			else if(Status == RMAX_ERR_HW_COMPLETION_ISSUE)
+			{
+				UE_LOG(LogRivermax, Error, TEXT("Completion issue while trying to commit next round of chunks."));
+				Listener->OnStreamError();
+				bIsActive = false;
+			}
+
+		} while (Status != RMAX_OK && bIsActive);
+	}
+
+	bool FRivermaxOutputStream::Init()
+	{
+		return true;
+	}
+
+	uint32 FRivermaxOutputStream::Run()
+	{
+		while (bIsActive)
+		{
+			if (CurrentFrame.IsValid() == false)
+			{
+				CurrentFrame = GetNextFrameToSend();
+			}
+
+			if(CurrentFrame.IsValid())
+			{
+				Process_AnyThread();
+			}
+			else
+			{
+				TRACE_CPUPROFILER_EVENT_SCOPE(FRivermaxOutputStream::WaitForFrame);
+				const double PreWaitTime = FPlatformTime::Seconds();
+				ReadyToSendEvent->Wait();
+				UE_LOG(LogRivermax, Verbose, TEXT("Outputstream waited %0.8f"), FPlatformTime::Seconds() - PreWaitTime);
+			}
+		}
+
+
+		DestroyStream();
+
+		return 0;
+	}
+
+	void FRivermaxOutputStream::Stop()
+	{
+		bIsActive = false;
+	}
+
+	void FRivermaxOutputStream::Exit()
+	{
+
+	}
+}
+
diff --git a/Engine/Plugins/VirtualProduction/Rivermax/RivermaxCore/Source/RivermaxCore/Private/RivermaxOutputStream.h b/Engine/Plugins/VirtualProduction/Rivermax/RivermaxCore/Source/RivermaxCore/Private/RivermaxOutputStream.h
new file mode 100644
--- /dev/null
+++ b/Engine/Plugins/VirtualProduction/Rivermax/RivermaxCore/Source/RivermaxCore/Private/RivermaxOutputStream.h
@@ -0,0 +1,137 @@
+// Copyright Epic Games, Inc. All Rights Reserved.
+
+#pragma once
+
+#include "IRivermaxOutputStream.h"
+
+#include "HAL/Runnable.h"
+#include "RivermaxHeader.h"
+#include "RivermaxOutputFrame.h"
+#include "RivermaxTypes.h"
+
+
+class FEvent;
+
+namespace UE::RivermaxCore::Private
+{
+	using UE::RivermaxCore::FRivermaxStreamOptions;
+	using UE::RivermaxCore::FRivermaxOutputVideoFrameInfo;
+	
+	// Todo make a proper RTP header struct
+	struct FRTPHeader
+	{
+		uint8 RawHeader[20];
+	};
+
+	struct FRivermaxOutputStreamMemory
+	{
+		uint16 PayloadSize = 0;
+		uint32 DataStrideSize = 1280;
+		uint32 HeaderStrideSize = 20;
+		uint32 LinesInChunk = 4;
+
+		uint32 PacketsInLine = 0;
+		uint32 ChunkSizeInStrides = 0;
+
+		uint32 FramesFieldPerMemoryBlock = 0;
+		uint32 PacketsInFrameField = 0;
+		uint32 PacketsPerMemoryBlock = 0;
+		uint32 ChunksPerFrameField = 0;
+		uint32 ChunksPerMemoryBlock = 0;
+		uint32 MemoryBlockCount = 0; 
+		uint32 StridesPerMemoryBlock = 0;
+
+		TArray<rmax_mem_block> MemoryBlocks;
+		TArray<uint16_t> PayloadSizes; //Array describing stride payload size
+		TArray<uint16_t> HeaderSizes; //Array describing header payload size
+		TArray<FRTPHeader> RTPHeaders;
+
+		rmax_buffer_attr BufferAttributes;
+	};
+
+	struct FRivermaxOutputStreamStats
+	{
+		uint32 ChunkRetries = 0;
+		uint32 TotalStrides = 0;
+		uint32 ChunkWait = 0;
+		uint32 CommitWaits = 0;
+		uint32 CommitRetries = 0;
+		uint64 MemoryBlockSentCounter = 0; //Global for an active capture session to track timestamp for next packet
+	};
+
+	struct FRivermaxOutputStreamData
+	{
+		/** Current sequence number being done */
+		uint32 SequenceNumber = 0;
+		double FrameFieldTimeIntervalNs = 0.0;
+		double StartSendTimeNs = 0.0;
+		double SendTimeNs = 0.0;
+		double InitialTimestampTick = 0.0;
+		bool bHasFrameFirstChunkBeenFetched = false;
+	};
+
+
+	class FRivermaxOutputStream : public UE::RivermaxCore::IRivermaxOutputStream, public FRunnable
+	{
+	public:
+		FRivermaxOutputStream();
+		virtual ~FRivermaxOutputStream();
+
+	public:
+
+		//~ Begin IRivermaxOutputStream interface
+		virtual bool Initialize(const FRivermaxStreamOptions& Options, IRivermaxOutputStreamListener& InListener) override;
+		virtual void Uninitialize() override;
+		virtual bool PushVideoFrame(const FRivermaxOutputVideoFrameInfo& NewFrame) override;
+		//~ End IRivermaxOutputStream interface
+
+		void Process_AnyThread();
+
+		//~ Begin FRunnable interface
+		virtual bool Init() override;
+		virtual uint32 Run() override;
+		virtual void Stop() override;
+		virtual void Exit() override;
+		//~ End FRunnable interface
+
+	private:
+		void InitializeBuffers();
+		void InitializeMemory();
+		void InitializeNextFrame(const TSharedPtr<FRivermaxOutputFrame>& NextFrame);
+		TSharedPtr<FRivermaxOutputFrame> GetNextFrameToSend();
+		TSharedPtr<FRivermaxOutputFrame> GetNextAvailableFrame(uint32 InFrameIdentifier);
+		void BuildRTPHeader(FRTPHeader& OutHeader) const;
+		void DestroyStream();
+		void WaitForNextRound(double NextRoundTime);
+		void GetNextChunk();
+		void SetupRTPHeaders();
+		void CommitNextChunks();
+
+	private:
+		FRivermaxStreamOptions Options;
+		FRivermaxOutputStreamMemory StreamMemory;
+		FRivermaxOutputStreamStats Stats;
+		FRivermaxOutputStreamData StreamData;
+
+		rmax_stream_id StreamId;
+		FCriticalSection FrameCriticalSection;
+
+		TSharedPtr<FRivermaxOutputFrame> CurrentFrame;
+
+		TArray<TSharedPtr<FRivermaxOutputFrame>> AvailableFrames;
+		TArray<TSharedPtr<FRivermaxOutputFrame>> FramesToSend;
+		double LastTime = 0.0;
+
+		TUniquePtr<FRunnableThread> RivermaxThread;
+		std::atomic<bool> bIsActive;
+
+		FEvent* ReadyToSendEvent = nullptr;
+
+		IRivermaxOutputStreamListener* Listener = nullptr;
+
+		static constexpr double MediaClockSampleRate = 90000.0; //Required to comply with SMTPE 2110-10.The Media Clock and RTP Clock rate for streams compliant to this standard shall be 90 kHz.
+		ERivermaxStreamType StreamType = ERivermaxStreamType::VIDEO_2110_20_STREAM; //todo
+	};
+}
+
+
diff --git a/Engine/Plugins/VirtualProduction/Rivermax/RivermaxCore/Source/RivermaxCore/Private/RivermaxUtils.cpp b/Engine/Plugins/VirtualProduction/Rivermax/RivermaxCore/Source/RivermaxCore/Private/RivermaxUtils.cpp
new file mode 100644
--- /dev/null
+++ b/Engine/Plugins/VirtualProduction/Rivermax/RivermaxCore/Source/RivermaxCore/Private/RivermaxUtils.cpp
@@ -0,0 +1,84 @@
+// Copyright Epic Games, Inc. All Rights Reserved.
+
+#include "RivermaxUtils.h"
+
+#include "RivermaxTypes.h"
+
+
+namespace UE::RivermaxCore::Private::Utils
+{
+	FString PixelFormatToSamplingDesc(ERivermaxOutputPixelFormat PixelFormat)
+	{
+		switch (PixelFormat)
+		{
+			case ERivermaxOutputPixelFormat::RMAX_10BIT_RGB:
+			case ERivermaxOutputPixelFormat::RMAX_8BIT_RGB:
+			case ERivermaxOutputPixelFormat::RMAX_16F_RGB:
+			{
+				return FString(TEXT("RGB"));
+			}
+			case ERivermaxOutputPixelFormat::RMAX_10BIT_YCBCR:
+			case ERivermaxOutputPixelFormat::RMAX_8BIT_YCBCR:
+			default:
+			{
+				return FString(TEXT("YCbCr-4:2:2"));
+			}
+		}
+	}
+
+	FString PixelFormatToBitDepth(ERivermaxOutputPixelFormat PixelFormat)
+	{
+		switch (PixelFormat)
+		{
+			case ERivermaxOutputPixelFormat::RMAX_10BIT_RGB:
+			case ERivermaxOutputPixelFormat::RMAX_10BIT_YCBCR:
+			{
+				return FString(TEXT("10"));
+			}
+			case ERivermaxOutputPixelFormat::RMAX_16F_RGB:
+			{
+				return FString(TEXT("16f"));
+			}
+			case ERivermaxOutputPixelFormat::RMAX_8BIT_RGB:
+			case ERivermaxOutputPixelFormat::RMAX_8BIT_YCBCR:
+			default:
+			{
+				return FString(TEXT("8"));
+			}
+		}
+	}
+
+	void StreamOptionsToSDPDescription(const UE::RivermaxCore::FRivermaxStreamOptions& Options, FAnsiStringBuilderBase& OutSDPDescription)
+	{
+		// Basic SDP string creation from a set of options. At some point, having a proper SDP loader / creator.
+		// Refer to https://datatracker.ietf.org/doc/html/rfc4570
+
+		FString FrameRateDescription;
+		if (FMath::IsNearlyZero(FMath::Frac(Options.FrameRate.AsDecimal())) == false)
+		{
+			FrameRateDescription = FString::Printf(TEXT("%d/%d"), Options.FrameRate.Numerator, Options.FrameRate.Denominator);
+		}
+		else
+		{
+			FrameRateDescription = FString::Printf(TEXT("%d"), (uint32)Options.FrameRate.AsDecimal());
+		}
+
+
+
+		OutSDPDescription.Appendf("v=0\n");
+		OutSDPDescription.Appendf("s=SMPTE ST2110 20 streams\n");
+		OutSDPDescription.Appendf("m=video %d RTP/AVP 96\n", Options.Port);
+		OutSDPDescription.Appendf("c=IN IP4 %S/64\n", *Options.DestinationAddress);
+		OutSDPDescription.Appendf("a=source-filter: incl IN IP4 %S %S\n", *Options.DestinationAddress, *Options.SourceAddress);
+		OutSDPDescription.Appendf("a=rtpmap:96 raw/90000\n");
+		OutSDPDescription.Appendf("a=fmtp: 96 sampling=%S; width=%d; height=%d; exactframerate=%S; depth=%S; TCS=SDR; colorimetry=BT709; PM=2110GPM; SSN=ST2110-20:2017; TP=2110TPN;\n"
+			, *PixelFormatToSamplingDesc(Options.PixelFormat)
+			, Options.Resolution.X
+			, Options.Resolution.Y
+			, *FrameRateDescription
+			, *PixelFormatToBitDepth(Options.PixelFormat));
+
+		OutSDPDescription.Appendf("a=mediaclk:direct=0\n");
+		OutSDPDescription.Appendf("a=mid:VID");
+	}
+}
\ No newline at end of file
diff --git a/Engine/Plugins/VirtualProduction/Rivermax/RivermaxCore/Source/RivermaxCore/Private/RivermaxUtils.h b/Engine/Plugins/VirtualProduction/Rivermax/RivermaxCore/Source/RivermaxCore/Private/RivermaxUtils.h
new file mode 100644
--- /dev/null
+++ b/Engine/Plugins/VirtualProduction/Rivermax/RivermaxCore/Source/RivermaxCore/Private/RivermaxUtils.h
@@ -0,0 +1,26 @@
+// Copyright Epic Games, Inc. All Rights Reserved.
+
+#pragma once
+
+#include "CoreMinimal.h"
+
+namespace UE::RivermaxCore
+{
+	struct FRivermaxStreamOptions;
+}
+
+namespace UE::RivermaxCore::Private::Utils
+{	
+	/** Various constants used for stream initialization */
+	static constexpr uint32 BytesPerHeader = 20;
+	static constexpr uint32 BytesPerPacket = 1200;
+	static constexpr uint32 PixelsPerGroup_422_8b = 2;
+	static constexpr uint32 BytesPerGroup_422_8b = 4;
+	static constexpr uint32 FullHDHeight = 1080;
+	static constexpr uint32 FullHDWidth = 1920;
+	static constexpr uint32 MaxPayloadSize = 1280;
+	static constexpr uint32 VideoTROModification = 0; // Todo : Investigate what TRO modification is
+
+	/** Convert a set of streaming option to its SDP description. Currently only support video type. */
+	void StreamOptionsToSDPDescription(const UE::RivermaxCore::FRivermaxStreamOptions& Options, FAnsiStringBuilderBase& OutSDPDescription);
+}
diff --git a/Engine/Plugins/VirtualProduction/Rivermax/RivermaxCore/Source/RivermaxCore/Public/IRivermaxCoreModule.h b/Engine/Plugins/VirtualProduction/Rivermax/RivermaxCore/Source/RivermaxCore/Public/IRivermaxCoreModule.h
new file mode 100644
--- /dev/null
+++ b/Engine/Plugins/VirtualProduction/Rivermax/RivermaxCore/Source/RivermaxCore/Public/IRivermaxCoreModule.h
@@ -0,0 +1,44 @@
+// Copyright Epic Games, Inc. All Rights Reserved.
+
+#pragma once
+
+#include "Modules/ModuleInterface.h"
+
+#include "IRivermaxInputStream.h"
+#include "IRivermaxOutputStream.h"
+#include "Modules/ModuleManager.h"
+
+
+namespace UE::RivermaxCore
+{
+	class IRivermaxManager;
+}
+
+
+/**
+ * Core module for Rivermax access from the engine. Users can create different stream types that are exposed to 
+ * get data flow ongoing.
+ */
+class IRivermaxCoreModule : public IModuleInterface
+{
+public:
+	static inline IRivermaxCoreModule& Get()
+	{
+		return FModuleManager::LoadModuleChecked<IRivermaxCoreModule>("RivermaxCore");
+	}
+
+	static inline bool IsAvailable()
+	{
+		return FModuleManager::Get().IsModuleLoaded("RivermaxCore");
+	}
+
+	/** Create input stream managing receiving data from rivermax */
+	virtual TUniquePtr<UE::RivermaxCore::IRivermaxInputStream> CreateInputStream() = 0;
+
+	/** Create output stream managing sending data to rivermax */
+	virtual TUniquePtr<UE::RivermaxCore::IRivermaxOutputStream> CreateOutputStream() = 0;
+
+	/** Returns Rivermax manager singleton to query for stats, status, etc... */
+	virtual TSharedPtr<UE::RivermaxCore::IRivermaxManager> GetRivermaxManager() = 0;
+};
+
diff --git a/Engine/Plugins/VirtualProduction/Rivermax/RivermaxCore/Source/RivermaxCore/Public/IRivermaxInputStream.h b/Engine/Plugins/VirtualProduction/Rivermax/RivermaxCore/Source/RivermaxCore/Public/IRivermaxInputStream.h
new file mode 100644
--- /dev/null
+++ b/Engine/Plugins/VirtualProduction/Rivermax/RivermaxCore/Source/RivermaxCore/Public/IRivermaxInputStream.h
@@ -0,0 +1,52 @@
+// Copyright Epic Games, Inc. All Rights Reserved.
+
+#pragma once
+
+#include "CoreMinimal.h"
+
+namespace UE::RivermaxCore
+{
+	struct FRivermaxStreamOptions;
+
+	struct RIVERMAXCORE_API FRivermaxInputVideoFrameDescriptor
+	{
+		uint32 Height = 0;
+		uint32 Width = 0;
+		uint32 Stride = 0;
+		uint32 VideoBufferSize = 0;
+	};
+
+	struct RIVERMAXCORE_API FRivermaxInputVideoFrameRequest
+	{
+		uint8* VideoBuffer = nullptr;
+	};
+
+	struct RIVERMAXCORE_API FRivermaxInputVideoFrameReception
+	{
+		uint8* VideoBuffer = nullptr;
+	};
+
+	class RIVERMAXCORE_API IRivermaxInputStreamListener
+	{
+	public:
+		/** Initialization completion callback with result */
+		virtual void OnInitializationCompleted(bool bHasSucceed) = 0;
+	
+		/** Called when stream is ready to fill the next frame. Returns true if a frame was successfully requested */
+		virtual bool OnVideoFrameRequested(const FRivermaxInputVideoFrameDescriptor& FrameInfo, FRivermaxInputVideoFrameRequest& OutVideoFrameRequest) = 0;
+		
+		/** Called when a frame has been received */
+		virtual void OnVideoFrameReceived(const FRivermaxInputVideoFrameDescriptor& FrameInfo, const FRivermaxInputVideoFrameReception& ReceivedVideoFrame) = 0;
+	};
+
+	class RIVERMAXCORE_API IRivermaxInputStream
+	{
+	public:
+		virtual ~IRivermaxInputStream() = default;
+
+	public:
+		virtual bool Initialize(const FRivermaxStreamOptions& InOptions, IRivermaxInputStreamListener& InListener) = 0;
+		virtual void Uninitialize() = 0;
+	};
+}
+
diff --git a/Engine/Plugins/VirtualProduction/Rivermax/RivermaxCore/Source/RivermaxCore/Public/IRivermaxManager.h b/Engine/Plugins/VirtualProduction/Rivermax/RivermaxCore/Source/RivermaxCore/Public/IRivermaxManager.h
new file mode 100644
--- /dev/null
+++ b/Engine/Plugins/VirtualProduction/Rivermax/RivermaxCore/Source/RivermaxCore/Public/IRivermaxManager.h
@@ -0,0 +1,24 @@
+// Copyright Epic Games, Inc. All Rights Reserved.
+
+#pragma once
+
+#include "CoreMinimal.h"
+
+namespace UE::RivermaxCore
+{
+	/**
+	 * Not doing much at the moment but should be used as a central point to register every stream created
+	 * and keep track of stats. Also manages initialization of the library.
+	 */
+	class RIVERMAXCORE_API IRivermaxManager
+	{
+	public:
+		virtual ~IRivermaxManager() = default;
+
+	public:
+		
+		/** Returns true if Rivermax has been initialized and is usable */
+		virtual bool IsInitialized() const = 0;
+	};
+}
+
diff --git a/Engine/Plugins/VirtualProduction/Rivermax/RivermaxCore/Source/RivermaxCore/Public/IRivermaxOutputStream.h b/Engine/Plugins/VirtualProduction/Rivermax/RivermaxCore/Source/RivermaxCore/Public/IRivermaxOutputStream.h
new file mode 100644
--- /dev/null
+++ b/Engine/Plugins/VirtualProduction/Rivermax/RivermaxCore/Source/RivermaxCore/Public/IRivermaxOutputStream.h
@@ -0,0 +1,53 @@
+// Copyright Epic Games, Inc. All Rights Reserved.
+
+#pragma once
+
+#include "CoreMinimal.h"
+
+
+namespace UE::RivermaxCore
+{
+	struct FRivermaxStreamOptions;
+
+	struct RIVERMAXCORE_API FRivermaxOutputVideoFrameInfo
+	{
+		uint32 FrameIdentifier = 0;
+		uint32 Height = 0;
+		uint32 Width = 0;
+		uint32 Stride = 0;
+		void* VideoBuffer = nullptr;
+	};
+
+	class RIVERMAXCORE_API IRivermaxOutputStreamListener
+	{
+	public:
+
+		/** Initialization completion callback with result */
+		virtual void OnInitializationCompleted(bool bHasSucceed) = 0;
+
+		/** Called when stream has encountered an error and has to stop */
+		virtual void OnStreamError() = 0;
+	};
+
+	class RIVERMAXCORE_API IRivermaxOutputStream
+	{
+	public:
+		virtual ~IRivermaxOutputStream() = default;
+
+	public:
+
+		/** 
+		 * Initializes stream using input options. Returns false if stream creation has failed. 
+		 * Initialization is completed when OnInitializationCompletion has been called with result.
+		 */
+		virtual bool Initialize(const FRivermaxStreamOptions& Options, IRivermaxOutputStreamListener& InListener) = 0;
+
+		/** Uninitializes current stream */
+		virtual void Uninitialize() = 0;
+
+		/** Pushes new video frame to the stream. Returns false if frame couldn't be pushed. */
+		virtual bool PushVideoFrame(const FRivermaxOutputVideoFrameInfo& NewFrame) = 0;
+	};
+}
+
+
diff --git a/Engine/Plugins/VirtualProduction/Rivermax/RivermaxCore/Source/RivermaxCore/Public/RivermaxHeader.h b/Engine/Plugins/VirtualProduction/Rivermax/RivermaxCore/Source/RivermaxCore/Public/RivermaxHeader.h
new file mode 100644
--- /dev/null
+++ b/Engine/Plugins/VirtualProduction/Rivermax/RivermaxCore/Source/RivermaxCore/Public/RivermaxHeader.h
@@ -0,0 +1,20 @@
+// Copyright Epic Games, Inc. All Rights Reserved.
+
+#pragma once
+
+THIRD_PARTY_INCLUDES_START
+
+#if PLATFORM_WINDOWS
+#include "Windows/AllowWindowsPlatformTypes.h"
+#include "Windows/PreWindowsApi.h"
+#endif
+
+#include "rivermax_api.h"
+
+#if PLATFORM_WINDOWS
+#include "Windows/PostWindowsApi.h"
+#include "Windows/HideWindowsPlatformTypes.h"
+#endif
+
+THIRD_PARTY_INCLUDES_END
+
diff --git a/Engine/Plugins/VirtualProduction/Rivermax/RivermaxCore/Source/RivermaxCore/Public/RivermaxTypes.h b/Engine/Plugins/VirtualProduction/Rivermax/RivermaxCore/Source/RivermaxCore/Public/RivermaxTypes.h
new file mode 100644
--- /dev/null
+++ b/Engine/Plugins/VirtualProduction/Rivermax/RivermaxCore/Source/RivermaxCore/Public/RivermaxTypes.h
@@ -0,0 +1,58 @@
+// Copyright Epic Games, Inc. All Rights Reserved.
+
+#pragma once
+
+#include "CoreMinimal.h"
+
+#include "Misc/FrameRate.h"
+
+namespace UE::RivermaxCore
+{
+	enum class RIVERMAXCORE_API ERivermaxOutputPixelFormat
+	{
+		RMAX_8BIT_YCBCR,
+		RMAX_8BIT_RGB,
+		RMAX_10BIT_YCBCR,
+		RMAX_10BIT_RGB,
+		RMAX_16F_RGB,
+	};
+
+	struct RIVERMAXCORE_API FRivermaxStreamOptions
+	{
+		/** Desired stream resolution */
+		FIntPoint Resolution = { 1920, 1080 };
+
+		/** Stream FrameRate */
+		FFrameRate FrameRate = { 24,1 };
+
+		/** Source IP to bind to */
+		FString SourceAddress;
+
+		/** Destination IP to send to. Defaults to multicast group IP. */
+		FString DestinationAddress = TEXT("224.1.1.1");
+
+		/** Port to be used by stream */
+		uint32 Port = 50000;
+
+		/** Desired stream pixel format */
+		ERivermaxOutputPixelFormat PixelFormat = ERivermaxOutputPixelFormat::RMAX_8BIT_YCBCR;
+
+		/** Sample count to buffer. */
+		int32 NumberOfBuffers = 2;
+
+		/** Stride in bytes for a row */
+		uint32 Stride = 0;
+	};
+
+	enum class RIVERMAXCORE_API ERivermaxStreamType : uint8
+	{
+		VIDEO_2110_20_STREAM,
+
+		/** Todo add additional stream types */
+		//VIDEO_2110_22_STREAM, compressed video
+		//AUDIO_2110_30_31_STREAM,
+		//ANCILLARY_2110_40_STREAM,
+	};
+}
+
+
diff --git a/Engine/Plugins/VirtualProduction/Rivermax/RivermaxCore/Source/RivermaxCore/RivermaxCore.build.cs b/Engine/Plugins/VirtualProduction/Rivermax/RivermaxCore/Source/RivermaxCore/RivermaxCore.build.cs
new file mode 100644
--- /dev/null
+++ b/Engine/Plugins/VirtualProduction/Rivermax/RivermaxCore/Source/RivermaxCore/RivermaxCore.build.cs
@@ -0,0 +1,41 @@
+// Copyright Epic Games, Inc. All Rights Reserved.
+
+using UnrealBuildTool;
+
+public class RivermaxCore : ModuleRules
+{
+	public RivermaxCore(ReadOnlyTargetRules Target) : base(Target)
+	{
+		PublicDependencyModuleNames.AddRange(
+			new string[] {
+				"Core",
+				"CoreUObject",
+				"MediaAssets",
+				"MediaIOCore",
+				"RivermaxLib",
+				"TimeManagement"
+			}
+		);
+
+		PrivateDependencyModuleNames.AddRange(
+			new string[] {
+				"Engine",
+				"RenderCore",
+				"RHI"
+			}
+		);
+
+		if (Target.bBuildEditor)
+		{
+			PublicDependencyModuleNames.AddRange(
+				new string[] {
+				}
+			);
+
+			PrivateDependencyModuleNames.AddRange(
+				new string[] {
+				}
+			);
+		}
+	}
+}
diff --git a/Engine/Plugins/VirtualProduction/Rivermax/RivermaxCore/Source/RivermaxRendering/Private/RivermaxRenderingModule.cpp b/Engine/Plugins/VirtualProduction/Rivermax/RivermaxCore/Source/RivermaxRendering/Private/RivermaxRenderingModule.cpp
new file mode 100644
--- /dev/null
+++ b/Engine/Plugins/VirtualProduction/Rivermax/RivermaxCore/Source/RivermaxRendering/Private/RivermaxRenderingModule.cpp
@@ -0,0 +1,17 @@
+// Copyright Epic Games, Inc. All Rights Reserved.
+
+#include "RivermaxRenderingModule.h"
+
+#include "Interfaces/IPluginManager.h"
+#include "Misc/Paths.h"
+#include "Modules/ModuleManager.h"
+#include "ShaderCore.h"
+
+
+void FRivermaxRenderingModule::StartupModule()
+{
+	const FString PluginShaderDir = FPaths::Combine(IPluginManager::Get().FindPlugin(TEXT("RivermaxCore"))->GetBaseDir(), TEXT("Shaders"));
+	AddShaderSourceDirectoryMapping(TEXT("/Plugin/RivermaxCore"), PluginShaderDir);
+
+}
+IMPLEMENT_MODULE(FRivermaxRenderingModule, RivermaxRendering);
diff --git a/Engine/Plugins/VirtualProduction/Rivermax/RivermaxCore/Source/RivermaxRendering/Private/RivermaxRenderingModule.h b/Engine/Plugins/VirtualProduction/Rivermax/RivermaxCore/Source/RivermaxRendering/Private/RivermaxRenderingModule.h
new file mode 100644
--- /dev/null
+++ b/Engine/Plugins/VirtualProduction/Rivermax/RivermaxCore/Source/RivermaxRendering/Private/RivermaxRenderingModule.h
@@ -0,0 +1,19 @@
+// Copyright Epic Games, Inc. All Rights Reserved.
+
+#pragma once
+
+#include "Modules/ModuleInterface.h"
+
+
+/**
+ * Module used to define Rivermax shaders used for pixel format conversion
+ */
+class FRivermaxRenderingModule : public IModuleInterface
+{
+public:
+	//~ Begin IModuleInterface
+	virtual void StartupModule() override;
+	//~ End IModuleInterface
+};
+
+
diff --git a/Engine/Plugins/VirtualProduction/Rivermax/RivermaxCore/Source/RivermaxRendering/Private/RivermaxShaders.cpp b/Engine/Plugins/VirtualProduction/Rivermax/RivermaxCore/Source/RivermaxRendering/Private/RivermaxShaders.cpp
new file mode 100644
--- /dev/null
+++ b/Engine/Plugins/VirtualProduction/Rivermax/RivermaxCore/Source/RivermaxRendering/Private/RivermaxShaders.cpp
@@ -0,0 +1,136 @@
+// Copyright Epic Games, Inc. All Rights Reserved.
+
+#include "RivermaxShaders.h"
+
+#include "RenderGraphBuilder.h"
+#include "RHIStaticStates.h"
+
+
+namespace UE::RivermaxShaders
+{
+	/* FRGBToYUV10Bit422LittleEndianCS shader
+	 *****************************************************************************/
+
+	IMPLEMENT_GLOBAL_SHADER(FRGBToYUV10Bit422LittleEndianCS, "/Plugin/RivermaxCore/Private/RivermaxShaders.usf", "RGBToYUV10Bit422", SF_Compute);
+
+	FRGBToYUV10Bit422LittleEndianCS::FParameters* FRGBToYUV10Bit422LittleEndianCS::AllocateAndSetParameters(FRDGBuilder& GraphBuilder, FRDGTextureRef RGBATexture, const FIntPoint& SourceSize, const FIntRect& SourceViewRect, const FIntPoint& OutputSize, const FMatrix& ColorTransform, const FVector& YUVOffset, bool bDoLinearToSrgb, FRDGBufferRef OutputBuffer)
+	{
+		FRGBToYUV10Bit422LittleEndianCS::FParameters* Parameters = GraphBuilder.AllocParameters<FRGBToYUV10Bit422LittleEndianCS::FParameters>();
+
+		const uint32 InputTextureSizeX = SourceSize.X;
+		const uint32 InputTextureSizeY = SourceSize.Y;
+		Parameters->RGBToYUVConversion.InputTexture = RGBATexture;
+		Parameters->RGBToYUVConversion.InputSampler = TStaticSamplerState<SF_Point>::GetRHI();
+		Parameters->RGBToYUVConversion.ColorTransform = (FMatrix44f)MediaShaders::CombineColorTransformAndOffset(ColorTransform, YUVOffset);
+		Parameters->RGBToYUVConversion.DoLinearToSrgb = bDoLinearToSrgb;
+		Parameters->RGBToYUVConversion.OnePixelDeltaX = 1.0f / (float)InputTextureSizeX;
+		Parameters->OnePixelDeltaY = 1.0f / (float)InputTextureSizeY;
+
+		// Used to offset indices to sample input texture based on ThreadId of CS
+		Parameters->InputPixelOffsetX = SourceViewRect.Min.X;
+		Parameters->InputPixelOffsetY = SourceViewRect.Min.Y;
+
+		// Output size will be based on RGB to YUV encoding ( divide by 2 ) and pixels encoded per element ( 4 pixels )
+		Parameters->HorizontalElementCount = OutputSize.X;
+
+		// Each thread will read many RGB pixels to convert to YUV
+		Parameters->InputTexturePixelsPerThread = InputTextureSizeX / OutputSize.X;
+
+		Parameters->OutYCbCrBuffer = GraphBuilder.CreateUAV(OutputBuffer);
+
+		return Parameters;
+	}
+
+	/* FRGBToRGB8BitCS shader
+	 *****************************************************************************/
+
+	IMPLEMENT_GLOBAL_SHADER(FRGBToRGB8BitCS, "/Plugin/RivermaxCore/Private/RivermaxShaders.usf", "RGBToRGB8BitCS", SF_Compute);
+
+	FRGBToRGB8BitCS::FParameters* FRGBToRGB8BitCS::AllocateAndSetParameters(FRDGBuilder& GraphBuilder, FRDGTextureRef RGBATexture, const FIntPoint& SourceSize, const FIntRect& SourceViewRect, const FIntPoint& OutputSize, FRDGBufferRef OutputBuffer)
+	{
+		FRGBToRGB8BitCS::FParameters* Parameters = GraphBuilder.AllocParameters<FRGBToRGB8BitCS::FParameters>();
+
+		const uint32 InputTextureSizeX = SourceSize.X;
+		const uint32 InputTextureSizeY = SourceSize.Y;
+		Parameters->InputTexture = RGBATexture;
+		Parameters->InputSampler = TStaticSamplerState<SF_Point>::GetRHI();
+		Parameters->OnePixelDeltaX = 1.0f / (float)InputTextureSizeX;
+		Parameters->OnePixelDeltaY = 1.0f / (float)InputTextureSizeY;
+
+		// Used to offset indices to sample input texture based on ThreadId of CS
+		Parameters->InputPixelOffsetX = SourceViewRect.Min.X;
+		Parameters->InputPixelOffsetY = SourceViewRect.Min.Y;
+
+		Parameters->HorizontalElementCount = OutputSize.X;
+
+		// Each thread will read many RGBA pixels to pack into RGB buffer
+		Parameters->InputTexturePixelsPerThread = InputTextureSizeX / OutputSize.X;
+
+		Parameters->OutRGB8Buffer = GraphBuilder.CreateUAV(OutputBuffer);
+
+		return Parameters;
+	}
+
+	/* FRGBToRGB10BitCS shader
+	 *****************************************************************************/
+
+	IMPLEMENT_GLOBAL_SHADER(FRGBToRGB10BitCS, "/Plugin/RivermaxCore/Private/RivermaxShaders.usf", "RGBToRGB10BitCS", SF_Compute);
+
+
+	FRGBToRGB10BitCS::FParameters* FRGBToRGB10BitCS::AllocateAndSetParameters(FRDGBuilder& GraphBuilder, FRDGTextureRef RGBATexture, const FIntPoint& SourceSize, const FIntRect& SourceViewRect, const FIntPoint& OutputSize, FRDGBufferRef OutputBuffer)
+	{
+		FRGBToRGB10BitCS::FParameters* Parameters = GraphBuilder.AllocParameters<FRGBToRGB10BitCS::FParameters>();
+
+		const uint32 InputTextureSizeX = SourceSize.X;
+		const uint32 InputTextureSizeY = SourceSize.Y;
+		Parameters->InputTexture = RGBATexture;
+		Parameters->InputSampler = TStaticSamplerState<SF_Point>::GetRHI();
+		Parameters->OnePixelDeltaX = 1.0f / (float)InputTextureSizeX;
+		Parameters->OnePixelDeltaY = 1.0f / (float)InputTextureSizeY;
+
+		// Used to offset indices to sample input texture based on ThreadId of CS
+		Parameters->InputPixelOffsetX = SourceViewRect.Min.X;
+		Parameters->InputPixelOffsetY = SourceViewRect.Min.Y;
+
+		Parameters->HorizontalElementCount = OutputSize.X;
+
+		// Each thread will read many RGBA pixels to pack into RGB buffer
+		Parameters->InputTexturePixelsPerThread = InputTextureSizeX / OutputSize.X;
+
+		Parameters->OutRGB10Buffer = GraphBuilder.CreateUAV(OutputBuffer);
+
+		return Parameters;
+	}
+
+	/* FRGBToRGB16fCS shader
+	 *****************************************************************************/
+
+	IMPLEMENT_GLOBAL_SHADER(FRGBToRGB16fCS, "/Plugin/RivermaxCore/Private/RivermaxShaders.usf", "RGBToRGB16fBitCS", SF_Compute);
+
+
+	FRGBToRGB16fCS::FParameters* FRGBToRGB16fCS::AllocateAndSetParameters(FRDGBuilder& GraphBuilder, FRDGTextureRef RGBATexture, const FIntPoint& SourceSize, const FIntRect& SourceViewRect, const FIntPoint& OutputSize, FRDGBufferRef OutputBuffer)
+	{
+		FRGBToRGB16fCS::FParameters* Parameters = GraphBuilder.AllocParameters<FRGBToRGB16fCS::FParameters>();
+
+		const uint32 InputTextureSizeX = SourceSize.X;
+		const uint32 InputTextureSizeY = SourceSize.Y;
+		Parameters->InputTexture = RGBATexture;
+		Parameters->InputSampler = TStaticSamplerState<SF_Point>::GetRHI();
+		Parameters->OnePixelDeltaX = 1.0f / (float)InputTextureSizeX;
+		Parameters->OnePixelDeltaY = 1.0f / (float)InputTextureSizeY;
+
+		// Used to offset indices to sample input texture based on ThreadId of CS
+		Parameters->InputPixelOffsetX = SourceViewRect.Min.X;
+		Parameters->InputPixelOffsetY = SourceViewRect.Min.Y;
+
+		Parameters->HorizontalElementCount = OutputSize.X;
+
+		// Each thread will read many RGBA pixels to pack into RGB buffer
+		Parameters->InputTexturePixelsPerThread = InputTextureSizeX / OutputSize.X;
+
+		Parameters->OutRGB16fBuffer = GraphBuilder.CreateUAV(OutputBuffer);
+
+		return Parameters;
+	}
+}
+
diff --git a/Engine/Plugins/VirtualProduction/Rivermax/RivermaxCore/Source/RivermaxRendering/Public/RivermaxShaders.h b/Engine/Plugins/VirtualProduction/Rivermax/RivermaxCore/Source/RivermaxRendering/Public/RivermaxShaders.h
new file mode 100644
--- /dev/null
+++ b/Engine/Plugins/VirtualProduction/Rivermax/RivermaxCore/Source/RivermaxRendering/Public/RivermaxShaders.h
@@ -0,0 +1,190 @@
+// Copyright Epic Games, Inc. All Rights Reserved.
+
+#pragma once
+
+#include "GlobalShader.h"
+
+#include "MediaShaders.h"
+#include "RenderGraphUtils.h"
+#include "RHI.h"
+#include "Shader.h"
+#include "ShaderParameterStruct.h"
+
+
+namespace UE::RivermaxShaders
+{
+	/**
+	 * Compute shader to convert RGB to packed YUV422 10 bits little endian
+	 *
+	 * This shader expects a single texture in PF_A2B10G10R10 format.
+	 */
+	class RIVERMAXRENDERING_API FRGBToYUV10Bit422LittleEndianCS : public FGlobalShader
+	{
+	public:
+		//Structure definition to match StructuredBuffer in RivermaxShaders.usf
+		struct FYUV10Bit422LEBuffer
+		{
+			uint32 DWord0;
+			uint32 DWord1;
+			uint32 DWord2;
+			uint32 DWord3;
+			uint32 DWord4;
+		};
+
+		DECLARE_GLOBAL_SHADER(FRGBToYUV10Bit422LittleEndianCS);
+
+		SHADER_USE_PARAMETER_STRUCT(FRGBToYUV10Bit422LittleEndianCS, FGlobalShader);
+
+		BEGIN_SHADER_PARAMETER_STRUCT(FParameters, )
+			SHADER_PARAMETER_STRUCT_INCLUDE(FRGBToYUVConversion, RGBToYUVConversion)
+			SHADER_PARAMETER(float, OnePixelDeltaY)
+			SHADER_PARAMETER(float, InputPixelOffsetX)
+			SHADER_PARAMETER(float, InputPixelOffsetY)
+			SHADER_PARAMETER(uint32, HorizontalElementCount)
+			SHADER_PARAMETER(uint32, InputTexturePixelsPerThread)
+			SHADER_PARAMETER_RDG_BUFFER_UAV(RWStructuredBuffer<FYUV10Bit422LEBuffer>, OutYCbCrBuffer)
+		END_SHADER_PARAMETER_STRUCT()
+
+	public:
+
+		// Called by the engine to determine which permutations to compile for this shader
+		static bool ShouldCompilePermutation(const FGlobalShaderPermutationParameters& Parameters)
+		{
+			return IsFeatureLevelSupported(Parameters.Platform, ERHIFeatureLevel::SM5);
+		}
+
+		/** Allocates and setup shader parameter in the incoming graph builder */
+		FRGBToYUV10Bit422LittleEndianCS::FParameters* AllocateAndSetParameters(FRDGBuilder& GraphBuilder, FRDGTextureRef RGBATexture, const FIntPoint& SourceSize, const FIntRect& SourceViewRect, const FIntPoint& OutputSize, const FMatrix& ColorTransform, const FVector& YUVOffset, bool bDoLinearToSrgb, FRDGBufferRef OutputBuffer);
+	};
+
+	/**
+	 * Compute shader to convert RGBA to packed 8 bits RGB
+	 */
+	class RIVERMAXRENDERING_API FRGBToRGB8BitCS : public FGlobalShader
+	{
+	public:
+		
+		// Structure definition to match StructuredBuffer in RivermaxShaders.usf
+		// For 8bit output, it's 3 bytes per pixels. To align with 4 bytes (32bit) we use 12 bytes (4 pixels) per output
+		struct FRGB8BitBuffer
+		{
+			uint32 DWord0;
+			uint32 DWord1;
+			uint32 DWord2;
+		};
+
+		DECLARE_GLOBAL_SHADER(FRGBToRGB8BitCS);
+
+		SHADER_USE_PARAMETER_STRUCT(FRGBToRGB8BitCS, FGlobalShader);
+
+		BEGIN_SHADER_PARAMETER_STRUCT(FParameters, )
+			SHADER_PARAMETER_RDG_TEXTURE(Texture2D, InputTexture)
+			SHADER_PARAMETER_SAMPLER(SamplerState, InputSampler)
+			SHADER_PARAMETER(float, OnePixelDeltaX)
+			SHADER_PARAMETER(float, OnePixelDeltaY)
+			SHADER_PARAMETER(float, InputPixelOffsetX)
+			SHADER_PARAMETER(float, InputPixelOffsetY)
+			SHADER_PARAMETER(uint32, HorizontalElementCount)
+			SHADER_PARAMETER(uint32, InputTexturePixelsPerThread)
+			SHADER_PARAMETER_RDG_BUFFER_UAV(RWStructuredBuffer<FRGB8BitBuffer>, OutRGB8Buffer)
+		END_SHADER_PARAMETER_STRUCT()
+
+	public:
+
+		// Called by the engine to determine which permutations to compile for this shader
+		static bool ShouldCompilePermutation(const FGlobalShaderPermutationParameters& Parameters)
+		{
+			return IsFeatureLevelSupported(Parameters.Platform, ERHIFeatureLevel::SM5);
+		}
+
+		/** Allocates and setup shader parameter in the incoming graph builder */
+		FRGBToRGB8BitCS::FParameters* AllocateAndSetParameters(FRDGBuilder& GraphBuilder, FRDGTextureRef RGBATexture, const FIntPoint& SourceSize, const FIntRect& SourceViewRect, const FIntPoint& OutputSize, FRDGBufferRef OutputBuffer);
+	};
+
+	/**
+	 * Compute shader to convert RGBA to packed 10 bits RGB
+	 */
+	class RIVERMAXRENDERING_API FRGBToRGB10BitCS : public FGlobalShader
+	{
+	public:
+
+		// Structure definition to match StructuredBuffer in RivermaxShaders.usf
+		// 10 bits per channel, 30 bits per pixel. Aligns on 15 32bits -> 480bits -> 16RGB10 pixels
+		// 2110-20 says it's 4 pixels per pgroup. 
+		struct FRGB10BitBuffer
+		{
+			uint32 DWords[15];
+		};
+
+		DECLARE_GLOBAL_SHADER(FRGBToRGB10BitCS);
+
+		SHADER_USE_PARAMETER_STRUCT(FRGBToRGB10BitCS, FGlobalShader);
+
+		BEGIN_SHADER_PARAMETER_STRUCT(FParameters, )
+			SHADER_PARAMETER_RDG_TEXTURE(Texture2D, InputTexture)
+			SHADER_PARAMETER_SAMPLER(SamplerState, InputSampler)
+			SHADER_PARAMETER(float, OnePixelDeltaX)
+			SHADER_PARAMETER(float, OnePixelDeltaY)
+			SHADER_PARAMETER(float, InputPixelOffsetX)
+			SHADER_PARAMETER(float, InputPixelOffsetY)
+			SHADER_PARAMETER(uint32, HorizontalElementCount)
+			SHADER_PARAMETER(uint32, InputTexturePixelsPerThread)
+			SHADER_PARAMETER_RDG_BUFFER_UAV(RWStructuredBuffer<FRGB10BitBuffer>, OutRGB10Buffer)
+		END_SHADER_PARAMETER_STRUCT()
+
+	public:
+
+		// Called by the engine to determine which permutations to compile for this shader
+		static bool ShouldCompilePermutation(const FGlobalShaderPermutationParameters& Parameters)
+		{
+			return IsFeatureLevelSupported(Parameters.Platform, ERHIFeatureLevel::SM5);
+		}
+
+		/** Allocates and setup shader parameter in the incoming graph builder */
+		FRGBToRGB10BitCS::FParameters* AllocateAndSetParameters(FRDGBuilder& GraphBuilder, FRDGTextureRef RGBATexture, const FIntPoint& SourceSize, const FIntRect& SourceViewRect, const FIntPoint& OutputSize, FRDGBufferRef OutputBuffer);
+	};
+
+	/**
+	 * Compute shader to convert RGBA to Float16 RGB
+	 */
+	class RIVERMAXRENDERING_API FRGBToRGB16fCS : public FGlobalShader
+	{
+	public:
+
+		DECLARE_GLOBAL_SHADER(FRGBToRGB16fCS);
+
+		SHADER_USE_PARAMETER_STRUCT(FRGBToRGB16fCS, FGlobalShader);
+
+		// Structure definition to match StructuredBuffer in RivermaxShaders.usf
+		// To align on 32bits, use 2 pixels (16 bit * 3 * 2 = 3x32bits)
+		struct FRGB16fBuffer
+		{
+			uint16 Pixel0[3];
+			uint16 Pixel1[3];
+		};
+
+		BEGIN_SHADER_PARAMETER_STRUCT(FParameters, )
+			SHADER_PARAMETER_RDG_TEXTURE(Texture2D, InputTexture)
+			SHADER_PARAMETER_SAMPLER(SamplerState, InputSampler)
+			SHADER_PARAMETER(float, OnePixelDeltaX)
+			SHADER_PARAMETER(float, OnePixelDeltaY)
+			SHADER_PARAMETER(float, InputPixelOffsetX)
+			SHADER_PARAMETER(float, InputPixelOffsetY)
+			SHADER_PARAMETER(uint32, HorizontalElementCount)
+			SHADER_PARAMETER(uint32, InputTexturePixelsPerThread)
+			SHADER_PARAMETER_RDG_BUFFER_UAV(RWStructuredBuffer<FRGB16fBuffer>, OutRGB16fBuffer)
+		END_SHADER_PARAMETER_STRUCT()
+
+	public:
+
+		// Called by the engine to determine which permutations to compile for this shader
+		static bool ShouldCompilePermutation(const FGlobalShaderPermutationParameters& Parameters)
+		{
+			return IsFeatureLevelSupported(Parameters.Platform, ERHIFeatureLevel::SM5);
+		}
+
+		/** Allocates and setup shader parameter in the incoming graph builder */
+		FRGBToRGB16fCS::FParameters* AllocateAndSetParameters(FRDGBuilder& GraphBuilder, FRDGTextureRef RGBATexture, const FIntPoint& SourceSize, const FIntRect& SourceViewRect, const FIntPoint& OutputSize, FRDGBufferRef OutputBuffer);
+	};
+}
+
diff --git a/Engine/Plugins/VirtualProduction/Rivermax/RivermaxCore/Source/RivermaxRendering/RivermaxRendering.build.cs b/Engine/Plugins/VirtualProduction/Rivermax/RivermaxCore/Source/RivermaxRendering/RivermaxRendering.build.cs
new file mode 100644
--- /dev/null
+++ b/Engine/Plugins/VirtualProduction/Rivermax/RivermaxCore/Source/RivermaxRendering/RivermaxRendering.build.cs
@@ -0,0 +1,23 @@
+// Copyright Epic Games, Inc. All Rights Reserved.
+
+using UnrealBuildTool;
+
+public class RivermaxRendering : ModuleRules
+{
+	public RivermaxRendering(ReadOnlyTargetRules Target) : base(Target)
+	{
+		PublicDependencyModuleNames.AddRange(
+			new string[] {
+			}
+		);
+
+		PrivateDependencyModuleNames.AddRange(
+			new string[] {
+				"Core",
+				"Projects",
+				"RenderCore",
+				"RHI"
+			}
+		);
+	}
+}
diff --git a/Engine/Plugins/VirtualProduction/Rivermax/RivermaxMedia/Source/RivermaxMedia/Private/RivermaxMediaCapture.cpp b/Engine/Plugins/VirtualProduction/Rivermax/RivermaxMedia/Source/RivermaxMedia/Private/RivermaxMediaCapture.cpp
--- a/Engine/Plugins/VirtualProduction/Rivermax/RivermaxMedia/Source/RivermaxMedia/Private/RivermaxMediaCapture.cpp
+++ b/Engine/Plugins/VirtualProduction/Rivermax/RivermaxMedia/Source/RivermaxMedia/Private/RivermaxMediaCapture.cpp
@@ -4,8 +4,10 @@
 
 #include "IRivermaxCoreModule.h"
 #include "IRivermaxOutputStream.h"
+#include "RenderGraphUtils.h"
 #include "RivermaxMediaLog.h"
 #include "RivermaxMediaOutput.h"
+#include "RivermaxShaders.h"
 #include "RivermaxTypes.h"
 #include "Slate/SceneViewport.h"
 #include "Widgets/SViewport.h"
@@ -44,6 +46,55 @@ namespace RivermaxMediaCaptureAnalytics
 }
 #endif
 
+namespace UE::RivermaxMediaCaptureUtil
+{
+	void GetOutputEncodingInfo(ERivermaxMediaOutputPixelFormat InPixelFormat, const FIntPoint& InSize, uint32& OutBytesPerElement, uint32& ElementsPerRow)
+	{
+		switch (InPixelFormat)
+		{
+		case ERivermaxMediaOutputPixelFormat::PF_10BIT_YCBCR_422:
+		{
+			// 2 RGB pixels per YUV pixels and 4 YUV pixels per entry in the buffer
+			constexpr uint32 RGBPixelsPerYUVPixels = 2;
+			const uint32 EncodedHorizontalResolution = InSize.X / RGBPixelsPerYUVPixels;
+			constexpr uint32 BytesPerYUVPixels = 5; //5 bytes per pixels for YUV422 10bits
+			const uint32 HorizontalByteCount = EncodedHorizontalResolution * BytesPerYUVPixels;
+			OutBytesPerElement = sizeof(UE::RivermaxShaders::FRGBToYUV10Bit422LittleEndianCS::FYUV10Bit422LEBuffer);
+			ElementsPerRow = HorizontalByteCount / OutBytesPerElement;
+			break;
+		}
+		case ERivermaxMediaOutputPixelFormat::PF_8BIT_RGB:
+		{
+			// 2 RGB pixels per YUV pixels and 4 YUV pixels per entry in the buffer
+			constexpr uint32 BytesPerRGBPixels = 3; //3 bytes per pixels RGB8
+			const uint32 HorizontalByteCount = InSize.X * BytesPerRGBPixels;
+			OutBytesPerElement = sizeof(UE::RivermaxShaders::FRGBToRGB8BitCS::FRGB8BitBuffer);
+			ElementsPerRow = HorizontalByteCount / OutBytesPerElement;
+			break;
+		}
+		case ERivermaxMediaOutputPixelFormat::PF_10BIT_RGB:
+		{
+			// 2 RGB pixels per YUV pixels and 4 YUV pixels per entry in the buffer
+			constexpr uint32 PixelPerGroup = 4;
+			constexpr uint32 BytesPerGroup = (PixelPerGroup * 3 * 10) / 8;
+			const uint32 HorizontalByteCount = (InSize.X / PixelPerGroup) * BytesPerGroup;
+			OutBytesPerElement = sizeof(UE::RivermaxShaders::FRGBToRGB10BitCS::FRGB10BitBuffer);
+			ElementsPerRow = HorizontalByteCount / OutBytesPerElement;
+			break;
+		}
+		case ERivermaxMediaOutputPixelFormat::PF_FLOAT16_RGB:
+		{
+			constexpr uint32 PixelPerGroup = 2;
+			constexpr uint32 BytesPerGroup = (PixelPerGroup * 3 * 16) / 8;
+			const uint32 HorizontalByteCount = (InSize.X / PixelPerGroup) * BytesPerGroup;
+			OutBytesPerElement = BytesPerGroup;
+			ElementsPerRow = HorizontalByteCount / OutBytesPerElement;
+			break;
+		}
+		}
+	}
+}
+
 
 ///* URivermaxMediaCapture implementation
 //*****************************************************************************/
@@ -105,7 +156,7 @@ void URivermaxMediaCapture::StopCaptureImpl(bool bAllowPendingFrameToBeProcess)
 	}
 }
 
-bool URivermaxMediaCapture::ShouldCaptureRHITexture() const
+bool URivermaxMediaCapture::ShouldCaptureRHIResource() const
 {
 	//todo for gpudirect support
 	return false;  
@@ -113,7 +164,7 @@ bool URivermaxMediaCapture::ShouldCaptureRHITexture() const
 
 void URivermaxMediaCapture::BeforeFrameCaptured_RenderingThread(const FCaptureBaseData& InBaseData, TSharedPtr<FMediaCaptureUserData, ESPMode::ThreadSafe> InUserData, FTextureRHIRef InTexture)
 {
-	if (ShouldCaptureRHITexture())
+	if (ShouldCaptureRHIResource())
 	{
 	
 	}
@@ -163,6 +214,31 @@ bool URivermaxMediaCapture::ConfigureStream(URivermaxMediaOutput* InMediaOutput,
 	case ERivermaxMediaOutputPixelFormat::PF_8BIT_YUV:
 	{
 		OutOptions.PixelFormat = ERivermaxOutputPixelFormat::RMAX_8BIT_YCBCR;
+		OutOptions.Stride = InMediaOutput->Resolution.X / 2 * 4; // 4 bytes for a group of 2 pixels
+		break;
+	}
+	case ERivermaxMediaOutputPixelFormat::PF_10BIT_YCBCR_422:
+	{
+		OutOptions.PixelFormat = ERivermaxOutputPixelFormat::RMAX_10BIT_YCBCR;
+		OutOptions.Stride = InMediaOutput->Resolution.X / 2 * 5; // 5 bytes for a group of 2 pixels (40bits / 8 = 5)
+		break;
+	}
+	case ERivermaxMediaOutputPixelFormat::PF_8BIT_RGB:
+	{
+		OutOptions.PixelFormat = ERivermaxOutputPixelFormat::RMAX_8BIT_RGB;
+		OutOptions.Stride = InMediaOutput->Resolution.X * 3; //3 bytes per pixel
+		break;
+	}
+	case ERivermaxMediaOutputPixelFormat::PF_10BIT_RGB:
+	{
+		OutOptions.PixelFormat = ERivermaxOutputPixelFormat::RMAX_10BIT_RGB;
+		OutOptions.Stride = InMediaOutput->Resolution.X / 4 * 15; //15 bytes for a group of 4 pixels (4 * 30bits / 8 = 15)
+		break;
+	}
+	case ERivermaxMediaOutputPixelFormat::PF_FLOAT16_RGB:
+	{
+		OutOptions.PixelFormat = ERivermaxOutputPixelFormat::RMAX_16F_RGB;
+		OutOptions.Stride = InMediaOutput->Resolution.X * 6; //3x 16bits per pixel (6 bytes)
 		break;
 	}
 	default:
@@ -190,11 +266,123 @@ void URivermaxMediaCapture::OnFrameCaptured_RenderingThread(const FCaptureBaseDa
 	RivermaxStream->PushVideoFrame(NewFrame);
 }
 
-void URivermaxMediaCapture::OnRHITextureCaptured_RenderingThread(const FCaptureBaseData& InBaseData,	TSharedPtr<FMediaCaptureUserData, ESPMode::ThreadSafe> InUserData, FTextureRHIRef InTexture)
+void URivermaxMediaCapture::OnRHIResourceCaptured_RenderingThread(const FCaptureBaseData& InBaseData,	TSharedPtr<FMediaCaptureUserData, ESPMode::ThreadSafe> InUserData, FTextureRHIRef InTexture)
 {
 	TRACE_CPUPROFILER_EVENT_SCOPE(URivermaxMediaCapture::OnFrameCaptured_RenderingThread);
 }
 
+FIntPoint URivermaxMediaCapture::GetCustomOutputSize(const FIntPoint& InSize) const
+{
+	URivermaxMediaOutput* RivermaxOutput = CastChecked<URivermaxMediaOutput>(MediaOutput);
+
+	uint32 BytesPerElement = 0;
+	uint32 ElementsPerRow = 0;
+	UE::RivermaxMediaCaptureUtil::GetOutputEncodingInfo(RivermaxOutput->PixelFormat, InSize, BytesPerElement, ElementsPerRow);
+	return FIntPoint(ElementsPerRow, InSize.Y);;
+}
+
+EMediaCaptureResourceType URivermaxMediaCapture::GetCustomOutputResourceType() const
+{
+	URivermaxMediaOutput* RivermaxOutput = CastChecked<URivermaxMediaOutput>(MediaOutput);
+
+	switch (RivermaxOutput->PixelFormat)
+	{
+	case ERivermaxMediaOutputPixelFormat::PF_10BIT_YCBCR_422:
+	case ERivermaxMediaOutputPixelFormat::PF_8BIT_RGB:
+	case ERivermaxMediaOutputPixelFormat::PF_10BIT_RGB:
+	case ERivermaxMediaOutputPixelFormat::PF_FLOAT16_RGB:
+	{
+		return EMediaCaptureResourceType::Buffer; //we use compute shader for all since output format doesn't match texture formats
+	}
+	default:
+		return EMediaCaptureResourceType::Texture;
+	}
+}
+
+FRDGBufferDesc URivermaxMediaCapture::GetCustomBufferDescription(const FIntPoint& InDesiredSize) const
+{
+	URivermaxMediaOutput* RivermaxOutput = CastChecked<URivermaxMediaOutput>(MediaOutput);
+
+	uint32 BytesPerElement = 0;
+	uint32 ElementsPerRow = 0;
+	UE::RivermaxMediaCaptureUtil::GetOutputEncodingInfo(RivermaxOutput->PixelFormat, InDesiredSize, BytesPerElement, ElementsPerRow);
+	return FRDGBufferDesc::CreateStructuredDesc(BytesPerElement, ElementsPerRow * InDesiredSize.Y);;
+}
+
+void URivermaxMediaCapture::OnCustomCapture_RenderingThread(FRDGBuilder& GraphBuilder, const FCaptureBaseData& InBaseData, TSharedPtr<FMediaCaptureUserData, ESPMode::ThreadSafe> InUserData, FRDGTextureRef InSourceTexture, FRDGBufferRef OutputBuffer, FResolveParams& ResolveParams, FVector2D CropU, FVector2D CropV)
+{
+	using namespace UE::RivermaxShaders;
+	URivermaxMediaOutput* RivermaxOutput = CastChecked<URivermaxMediaOutput>(MediaOutput);
+
+	FGlobalShaderMap* GlobalShaderMap = GetGlobalShaderMap(GMaxRHIFeatureLevel);
+
+	// Rectangle area to use from source. This is used when source render target is bigger than output resolution
+	const FIntRect ViewRect(ResolveParams.Rect.X1, ResolveParams.Rect.Y1, ResolveParams.Rect.X2, ResolveParams.Rect.Y2);
+	constexpr bool bDoLinearToSRGB = false;
+	const FIntPoint SourceSize = { InSourceTexture->Desc.Extent.X, InSourceTexture->Desc.Extent.Y };
+	const FIntVector GroupCount = FComputeShaderUtils::GetGroupCount(DesiredOutputSize, FComputeShaderUtils::kGolden2DGroupSize);
+	
+	switch (RivermaxOutput->PixelFormat)
+	{
+	case ERivermaxMediaOutputPixelFormat::PF_10BIT_YCBCR_422:
+	{
+		TShaderMapRef<FRGBToYUV10Bit422LittleEndianCS> ComputeShader(GlobalShaderMap);
+		FRGBToYUV10Bit422LittleEndianCS::FParameters* Parameters = ComputeShader->AllocateAndSetParameters(GraphBuilder, InSourceTexture, SourceSize, ViewRect, DesiredOutputSize, MediaShaders::RgbToYuvRec709Scaled, MediaShaders::YUVOffset10bits, bDoLinearToSRGB, OutputBuffer);
+
+		FComputeShaderUtils::AddPass(
+			GraphBuilder
+			, RDG_EVENT_NAME("RGBAToYUV10Bit422LE")
+			, ComputeShader
+			, Parameters
+			, GroupCount);
+	
+		break;
+	}
+	case ERivermaxMediaOutputPixelFormat::PF_8BIT_RGB:
+	{
+		TShaderMapRef<FRGBToRGB8BitCS> ComputeShader(GlobalShaderMap);
+		FRGBToRGB8BitCS::FParameters* Parameters = ComputeShader->AllocateAndSetParameters(GraphBuilder, InSourceTexture, SourceSize, ViewRect, DesiredOutputSize, OutputBuffer);
+
+		FComputeShaderUtils::AddPass(
+			GraphBuilder
+			, RDG_EVENT_NAME("RGBAToRGB8Bit")
+			, ComputeShader
+			, Parameters
+			, GroupCount);
+		
+		break;
+	}
+	case ERivermaxMediaOutputPixelFormat::PF_10BIT_RGB:
+	{
+		TShaderMapRef<FRGBToRGB10BitCS> ComputeShader(GlobalShaderMap);
+		FRGBToRGB10BitCS::FParameters* Parameters = ComputeShader->AllocateAndSetParameters(GraphBuilder, InSourceTexture, SourceSize, ViewRect, DesiredOutputSize, OutputBuffer);
+
+		FComputeShaderUtils::AddPass(
+			GraphBuilder
+			, RDG_EVENT_NAME("RGBAToRGB10Bit")
+			, ComputeShader
+			, Parameters
+			, GroupCount);
+
+		break;
+	}
+	case ERivermaxMediaOutputPixelFormat::PF_FLOAT16_RGB:
+	{
+		TShaderMapRef<FRGBToRGB16fCS> ComputeShader(GlobalShaderMap);
+		FRGBToRGB16fCS::FParameters* Parameters = ComputeShader->AllocateAndSetParameters(GraphBuilder, InSourceTexture, SourceSize, ViewRect, DesiredOutputSize, OutputBuffer);
+
+		FComputeShaderUtils::AddPass(
+			GraphBuilder
+			, RDG_EVENT_NAME("RGBAToRGB16f")
+			, ComputeShader
+			, Parameters
+			, GroupCount);
+
+		break;
+	}
+	}
+}
+
 void URivermaxMediaCapture::OnInitializationCompleted(bool bHasSucceed)
 {
 	if (GetState() != EMediaCaptureState::Stopped)
diff --git a/Engine/Plugins/VirtualProduction/Rivermax/RivermaxMedia/Source/RivermaxMedia/Private/RivermaxMediaOutput.cpp b/Engine/Plugins/VirtualProduction/Rivermax/RivermaxMedia/Source/RivermaxMedia/Private/RivermaxMediaOutput.cpp
--- a/Engine/Plugins/VirtualProduction/Rivermax/RivermaxMedia/Source/RivermaxMedia/Private/RivermaxMediaOutput.cpp
+++ b/Engine/Plugins/VirtualProduction/Rivermax/RivermaxMedia/Source/RivermaxMedia/Private/RivermaxMediaOutput.cpp
@@ -26,12 +26,35 @@ FIntPoint URivermaxMediaOutput::GetRequestedSize() const
 
 EPixelFormat URivermaxMediaOutput::GetRequestedPixelFormat() const
 {
-	return EPixelFormat::PF_B8G8R8A8;
+	EPixelFormat Result = EPixelFormat::PF_A2B10G10R10;
+	switch (PixelFormat)
+	{
+	case ERivermaxMediaOutputPixelFormat::PF_8BIT_YUV:
+		Result = EPixelFormat::PF_B8G8R8A8; //To be updated to use compute shader
+		break;
+	default: // All pixel formats use output buffer as their output so no need for a texture format
+		Result = EPixelFormat::PF_A2B10G10R10;
+		break;
+	}
+	return Result;
 }
 
 EMediaCaptureConversionOperation URivermaxMediaOutput::GetConversionOperation(EMediaCaptureSourceType InSourceType) const
 {
-	const EMediaCaptureConversionOperation Result = EMediaCaptureConversionOperation::RGBA8_TO_YUV_8BIT;
+	EMediaCaptureConversionOperation Result = EMediaCaptureConversionOperation::CUSTOM;
+	switch (PixelFormat)
+	{
+	case ERivermaxMediaOutputPixelFormat::PF_8BIT_YUV:
+		Result = EMediaCaptureConversionOperation::RGBA8_TO_YUV_8BIT;
+		break;
+	case ERivermaxMediaOutputPixelFormat::PF_10BIT_YCBCR_422:
+	case ERivermaxMediaOutputPixelFormat::PF_8BIT_RGB:
+	case ERivermaxMediaOutputPixelFormat::PF_10BIT_RGB:
+	case ERivermaxMediaOutputPixelFormat::PF_FLOAT16_RGB:
+	default:
+		Result = EMediaCaptureConversionOperation::CUSTOM; //We handle all conversion for rivermax since it's really tied to endianness of 2110
+		break;
+	}
 	return Result;
 }
 
diff --git a/Engine/Plugins/VirtualProduction/Rivermax/RivermaxMedia/Source/RivermaxMedia/Public/RivermaxMediaCapture.h b/Engine/Plugins/VirtualProduction/Rivermax/RivermaxMedia/Source/RivermaxMedia/Public/RivermaxMediaCapture.h
--- a/Engine/Plugins/VirtualProduction/Rivermax/RivermaxMedia/Source/RivermaxMedia/Public/RivermaxMediaCapture.h
+++ b/Engine/Plugins/VirtualProduction/Rivermax/RivermaxMedia/Source/RivermaxMedia/Public/RivermaxMediaCapture.h
@@ -37,11 +37,17 @@ protected:
 	virtual bool UpdateSceneViewportImpl(TSharedPtr<FSceneViewport>& InSceneViewport) override;
 	virtual bool UpdateRenderTargetImpl(UTextureRenderTarget2D* InRenderTarget) override;
 	virtual void StopCaptureImpl(bool bAllowPendingFrameToBeProcess) override;
-	virtual bool ShouldCaptureRHITexture() const override;
+	virtual bool ShouldCaptureRHIResource() const override;
 	
 	virtual void BeforeFrameCaptured_RenderingThread(const FCaptureBaseData& InBaseData, TSharedPtr<FMediaCaptureUserData, ESPMode::ThreadSafe> InUserData, FTextureRHIRef InTexture) override;
 	virtual void OnFrameCaptured_RenderingThread(const FCaptureBaseData& InBaseData, TSharedPtr<FMediaCaptureUserData, ESPMode::ThreadSafe> InUserData, void* InBuffer, int32 Width, int32 Height, int32 BytesPerRow) override;
-	virtual void OnRHITextureCaptured_RenderingThread(const FCaptureBaseData& InBaseData, TSharedPtr<FMediaCaptureUserData, ESPMode::ThreadSafe> InUserData, FTextureRHIRef InTexture) override;
+	virtual void OnRHIResourceCaptured_RenderingThread(const FCaptureBaseData& InBaseData, TSharedPtr<FMediaCaptureUserData, ESPMode::ThreadSafe> InUserData, FTextureRHIRef InTexture) override;
+	
+	/** For custom conversion, methods that need to be overriden */
+	virtual FIntPoint GetCustomOutputSize(const FIntPoint& InSize) const override;
+	virtual EMediaCaptureResourceType GetCustomOutputResourceType() const override;
+	virtual FRDGBufferDesc GetCustomBufferDescription(const FIntPoint& InDesiredSize) const override;
+	virtual void OnCustomCapture_RenderingThread(FRDGBuilder& GraphBuilder, const FCaptureBaseData& InBaseData, TSharedPtr<FMediaCaptureUserData, ESPMode::ThreadSafe> InUserData, FRDGTextureRef InSourceTexture, FRDGBufferRef OutputBuffer, FResolveParams& ResolveParams, FVector2D CropU, FVector2D CropV) override;
 	//~ End UMediaCapture interface
 
 	//~ Begin IRivermaxOutputStreamListener interface
diff --git a/Engine/Plugins/VirtualProduction/Rivermax/RivermaxMedia/Source/RivermaxMedia/Public/RivermaxMediaOutput.h b/Engine/Plugins/VirtualProduction/Rivermax/RivermaxMedia/Source/RivermaxMedia/Public/RivermaxMediaOutput.h
--- a/Engine/Plugins/VirtualProduction/Rivermax/RivermaxMedia/Source/RivermaxMedia/Public/RivermaxMediaOutput.h
+++ b/Engine/Plugins/VirtualProduction/Rivermax/RivermaxMedia/Source/RivermaxMedia/Public/RivermaxMediaOutput.h
@@ -14,7 +14,10 @@ UENUM()
 enum class ERivermaxMediaOutputPixelFormat : uint8
 {
 	PF_8BIT_YUV UMETA(DisplayName = "8bit YUV"),
-	//Todo add 10bit YUV, 8/10 bit RGB
+	PF_10BIT_YCBCR_422 UMETA(DisplayName = "10bit YCbCr 4:2:2"),
+	PF_8BIT_RGB UMETA(DisplayName = "8bit RGB"),
+	PF_10BIT_RGB UMETA(DisplayName = "10bit RGB"),
+	PF_FLOAT16_RGB UMETA(DisplayName = "16bit Float RGB")
 };
 
 
diff --git a/Engine/Plugins/VirtualProduction/Rivermax/RivermaxMedia/Source/RivermaxMedia/RivermaxMedia.build.cs b/Engine/Plugins/VirtualProduction/Rivermax/RivermaxMedia/Source/RivermaxMedia/RivermaxMedia.build.cs
--- a/Engine/Plugins/VirtualProduction/Rivermax/RivermaxMedia/Source/RivermaxMedia/RivermaxMedia.build.cs
+++ b/Engine/Plugins/VirtualProduction/Rivermax/RivermaxMedia/Source/RivermaxMedia/RivermaxMedia.build.cs
@@ -21,7 +21,8 @@ public class RivermaxMedia : ModuleRules
 				"Engine",
 				"RenderCore",
 				"RHI",
-				"RivermaxCore"
+				"RivermaxCore",
+				"RivermaxRendering"
 			}
 		);
 
diff --git a/Engine/Shaders/Private/MediaShaders.usf b/Engine/Shaders/Private/MediaShaders.usf
--- a/Engine/Shaders/Private/MediaShaders.usf
+++ b/Engine/Shaders/Private/MediaShaders.usf
@@ -332,52 +332,6 @@ void YVYUConvertPS(
 	OutColor.w = 1.0f;
 }
 
-// shader from 8 bits RGB to 8 bits UYVY
-void RGB8toUYVY8ConvertPS(
-	float4 InPosition : SV_POSITION,
-	float2 InUV : TEXCOORD0,
-	out float4 OutColor : SV_Target0)
-{
-	float3 RGB0 = RGB8toUYVY8ConvertUB.Texture.Sample(RGB8toUYVY8ConvertUB.SamplerP, float2(InUV.x + 0.5f * RGB8toUYVY8ConvertUB.OnePixelDeltaX, InUV.y)).rgb;
-	float3 RGB1 = RGB8toUYVY8ConvertUB.Texture.Sample(RGB8toUYVY8ConvertUB.SamplerP, float2(InUV.x + 1.5f * RGB8toUYVY8ConvertUB.OnePixelDeltaX, InUV.y)).rgb;
-
-	float3 YUV0 = RgbToYuv(RGB0, RGB8toUYVY8ConvertUB.ColorTransform, RGB8toUYVY8ConvertUB.LinearToSrgb);
-	float3 YUV1 = RgbToYuv(RGB1, RGB8toUYVY8ConvertUB.ColorTransform, RGB8toUYVY8ConvertUB.LinearToSrgb);
-	
-	OutColor.xyzw = float4(YUV0.z, YUV0.x, YUV0.y, YUV1.x);
-}
-
-
-// shader to convert RGB to YUV v210
-void RGB10toYUVv210ConvertPS(
-	float4 InPosition : SV_POSITION,
-	float2 InUV : TEXCOORD0,
-	out uint4 OutColor : SV_Target0)
-{
-	float X = (InUV.x * RGB10toYUVv210ConvertUB.PaddingScale) - RGB10toYUVv210ConvertUB.OnePixelDeltaX * 2.5f;
-
-	float3 RGB0 = RGB10toYUVv210ConvertUB.Texture.Sample(RGB10toYUVv210ConvertUB.SamplerP, float2(X, InUV.y)).rgb;
-	float3 RGB1 = RGB10toYUVv210ConvertUB.Texture.Sample(RGB10toYUVv210ConvertUB.SamplerP, float2(X + RGB10toYUVv210ConvertUB.OnePixelDeltaX, InUV.y)).rgb;
-	float3 RGB2 = RGB10toYUVv210ConvertUB.Texture.Sample(RGB10toYUVv210ConvertUB.SamplerP, float2(X + RGB10toYUVv210ConvertUB.OnePixelDeltaX * 2.0f, InUV.y)).rgb;
-	float3 RGB3 = RGB10toYUVv210ConvertUB.Texture.Sample(RGB10toYUVv210ConvertUB.SamplerP, float2(X + RGB10toYUVv210ConvertUB.OnePixelDeltaX * 3.0f, InUV.y)).rgb;
-	float3 RGB4 = RGB10toYUVv210ConvertUB.Texture.Sample(RGB10toYUVv210ConvertUB.SamplerP, float2(X + RGB10toYUVv210ConvertUB.OnePixelDeltaX * 4.0f, InUV.y)).rgb;
-	float3 RGB5 = RGB10toYUVv210ConvertUB.Texture.Sample(RGB10toYUVv210ConvertUB.SamplerP, float2(X + RGB10toYUVv210ConvertUB.OnePixelDeltaX * 5.0f, InUV.y)).rgb;
-
-	uint3 YUV0 = RgbToYuv(RGB0, RGB10toYUVv210ConvertUB.ColorTransform, RGB10toYUVv210ConvertUB.LinearToSrgb) * 1023;
-	uint3 YUV1 = RgbToYuv(RGB1, RGB10toYUVv210ConvertUB.ColorTransform, RGB10toYUVv210ConvertUB.LinearToSrgb) * 1023;
-	uint3 YUV2 = RgbToYuv(RGB2, RGB10toYUVv210ConvertUB.ColorTransform, RGB10toYUVv210ConvertUB.LinearToSrgb) * 1023;
-	uint3 YUV3 = RgbToYuv(RGB3, RGB10toYUVv210ConvertUB.ColorTransform, RGB10toYUVv210ConvertUB.LinearToSrgb) * 1023;
-	uint3 YUV4 = RgbToYuv(RGB4, RGB10toYUVv210ConvertUB.ColorTransform, RGB10toYUVv210ConvertUB.LinearToSrgb) * 1023;
-	uint3 YUV5 = RgbToYuv(RGB5, RGB10toYUVv210ConvertUB.ColorTransform, RGB10toYUVv210ConvertUB.LinearToSrgb) * 1023;
-
-	uint W0 = (YUV0.z << 20) | (YUV0.x << 10) | YUV0.y;
-	uint W1 = (YUV2.x << 20) | (YUV2.y << 10) | YUV1.x;
-	uint W2 = (YUV4.y << 20) | (YUV3.x << 10) | YUV2.z;
-	uint W3 = (YUV5.x << 20) | (YUV4.z << 10) | YUV4.x;
-
-	OutColor.xyzw = uint4(W0, W1, W2, W3);
-}
-
 
 // shader to convert RGB to Y component of YUV
 void RGB8toY8ConvertPS(float4 InPosition : SV_POSITION,
@@ -410,6 +364,64 @@ void ReadTextureExternalPS(
     OutColor = float4(sRGBToLinear(InColor.xyz), InColor.w);
 }
 
+/* Media Capture Pixel shaders 
+=============================================================================*/
+
+Texture2D InputTexture;
+SamplerState InputSampler;
+float4x4 ColorTransform;
+uint DoLinearToSrgb;
+float OnePixelDeltaX;
+float PaddingScale;
+float InputPixelOffsetX;
+
+
+// shader to convert from RGB to 8 bits UYVY
+void RGB8toUYVY8ConvertPS(
+	noperspective float4 UVAndScreenPos : TEXCOORD0,
+	out float4 OutColor : SV_Target0)
+{
+	float2 UV = UVAndScreenPos.xy;
+	float3 RGB0 = InputTexture.Sample(InputSampler, float2(UV.x + 0.5f * OnePixelDeltaX, UV.y)).rgb;
+	float3 RGB1 = InputTexture.Sample(InputSampler, float2(UV.x + 1.5f * OnePixelDeltaX, UV.y)).rgb;
+
+	float3 YUV0 = RgbToYuv(RGB0, ColorTransform, DoLinearToSrgb);
+	float3 YUV1 = RgbToYuv(RGB1, ColorTransform, DoLinearToSrgb);
+	
+	OutColor.xyzw = float4(YUV0.z, YUV0.x, YUV0.y, YUV1.x);
+}
+
+
+// shader to convert from RGB to YUV v210
+void RGB10toYUVv210ConvertPS(
+	noperspective float4 UVAndScreenPos : TEXCOORD0,
+	out uint4 OutColor : SV_Target0)
+{
+	float2 UV = UVAndScreenPos.xy;
+	float X = (UV.x * PaddingScale) - OnePixelDeltaX * 2.5f;
+
+	float3 RGB0 = InputTexture.Sample(InputSampler, float2(X, UV.y)).rgb;
+	float3 RGB1 = InputTexture.Sample(InputSampler, float2(X + OnePixelDeltaX, UV.y)).rgb;
+	float3 RGB2 = InputTexture.Sample(InputSampler, float2(X + OnePixelDeltaX * 2.0f, UV.y)).rgb;
+	float3 RGB3 = InputTexture.Sample(InputSampler, float2(X + OnePixelDeltaX * 3.0f, UV.y)).rgb;
+	float3 RGB4 = InputTexture.Sample(InputSampler, float2(X + OnePixelDeltaX * 4.0f, UV.y)).rgb;
+	float3 RGB5 = InputTexture.Sample(InputSampler, float2(X + OnePixelDeltaX * 5.0f, UV.y)).rgb;
+
+	uint3 YUV0 = RgbToYuv(RGB0, ColorTransform, DoLinearToSrgb) * 1023;
+	uint3 YUV1 = RgbToYuv(RGB1, ColorTransform, DoLinearToSrgb) * 1023;
+	uint3 YUV2 = RgbToYuv(RGB2, ColorTransform, DoLinearToSrgb) * 1023;
+	uint3 YUV3 = RgbToYuv(RGB3, ColorTransform, DoLinearToSrgb) * 1023;
+	uint3 YUV4 = RgbToYuv(RGB4, ColorTransform, DoLinearToSrgb) * 1023;
+	uint3 YUV5 = RgbToYuv(RGB5, ColorTransform, DoLinearToSrgb) * 1023;
+
+	uint W0 = (YUV0.z << 20) | (YUV0.x << 10) | YUV0.y;
+	uint W1 = (YUV2.x << 20) | (YUV2.y << 10) | YUV1.x;
+	uint W2 = (YUV4.y << 20) | (YUV3.x << 10) | YUV2.z;
+	uint W3 = (YUV5.x << 20) | (YUV4.z << 10) | YUV4.x;
+
+	OutColor.xyzw = uint4(W0, W1, W2, W3);
+}
+
 
 // Following defines must reflect values in EMediaCaptureConversionOperation in MediaOutput
 #define INVERT_ALPHA 3
@@ -417,11 +429,11 @@ void ReadTextureExternalPS(
 
 // Shader to swizzle rgba textures, set alpha to 1 or invert alpha
 void SwizzleRgbPS(
-	float4 InPosition : SV_POSITION,
-	float2 InUV : TEXCOORD0,
+	noperspective float4 UVAndScreenPos : TEXCOORD0,
 	out float4 OutColor : SV_Target0)
 {
-	float4 Color = SwizzleRgbUB.Texture.Sample(SwizzleRgbUB.SamplerP, InUV);
+	float2 UV = UVAndScreenPos.xy;
+	float4 Color = InputTexture.Sample(InputSampler, UV);
 #if CONVERSION_OP == INVERT_ALPHA
 	OutColor.rgba = float4(Color.rgb, 1 - Color.a);
 #elif CONVERSION_OP == SET_ALPHA_ONE
@@ -431,3 +443,4 @@ void SwizzleRgbPS(
 #endif
 }
 
+
diff --git a/Engine/Source/Runtime/RenderCore/Private/MediaShaders.cpp b/Engine/Source/Runtime/RenderCore/Private/MediaShaders.cpp
--- a/Engine/Source/Runtime/RenderCore/Private/MediaShaders.cpp
+++ b/Engine/Source/Runtime/RenderCore/Private/MediaShaders.cpp
@@ -1,8 +1,11 @@
 // Copyright Epic Games, Inc. All Rights Reserved.
 
 #include "MediaShaders.h"
-#include "ShaderParameterUtils.h"
+#include "RenderGraphBuilder.h"
 #include "RHIStaticStates.h"
+#include "ShaderParameterUtils.h"
+
+
 
 
 TGlobalResource<FMediaVertexDeclaration> GMediaVertexDeclaration;
@@ -529,75 +532,6 @@ void FYVYUConvertPS::SetParameters(FRHICommandList& CommandList, TRefCountPtr<FR
 }
 
 
-/* FRGB8toUYVY8ConvertPS shader
- *****************************************************************************/
-
-BEGIN_GLOBAL_SHADER_PARAMETER_STRUCT(FRGB8toUYVY8ConvertUB, )
-SHADER_PARAMETER(FMatrix44f, ColorTransform)
-SHADER_PARAMETER(uint32, LinearToSrgb)
-SHADER_PARAMETER(float, OnePixelDeltaX)
-SHADER_PARAMETER_TEXTURE(Texture2D, Texture)
-SHADER_PARAMETER_SAMPLER(SamplerState, SamplerP)
-END_GLOBAL_SHADER_PARAMETER_STRUCT()
-
-IMPLEMENT_GLOBAL_SHADER_PARAMETER_STRUCT(FRGB8toUYVY8ConvertUB, "RGB8toUYVY8ConvertUB");
-IMPLEMENT_SHADER_TYPE(, FRGB8toUYVY8ConvertPS, TEXT("/Engine/Private/MediaShaders.usf"), TEXT("RGB8toUYVY8ConvertPS"), SF_Pixel);
-
-
-void FRGB8toUYVY8ConvertPS::SetParameters(FRHICommandList& CommandList, TRefCountPtr<FRHITexture2D> RGBATexture, const FMatrix& ColorTransform, const FVector& YUVOffset, bool LinearToSrgb)
-{
-	FRGB8toUYVY8ConvertUB UB;
-	{
-		UB.ColorTransform = (FMatrix44f)MediaShaders::CombineColorTransformAndOffset(ColorTransform, YUVOffset);
-		UB.SamplerP = TStaticSamplerState<SF_Point>::GetRHI();
-		UB.LinearToSrgb = LinearToSrgb;
-		UB.Texture = RGBATexture;
-		UB.OnePixelDeltaX = 1.0f / float(RGBATexture->GetSizeX());
-	}
-
-	TUniformBufferRef<FRGB8toUYVY8ConvertUB> Data = TUniformBufferRef<FRGB8toUYVY8ConvertUB>::CreateUniformBufferImmediate(UB, UniformBuffer_SingleFrame);
-	SetUniformBufferParameter(CommandList, CommandList.GetBoundPixelShader(), GetUniformBufferParameter<FRGB8toUYVY8ConvertUB>(), Data);
-}
-
-
-/* FRGB10toYUVv210ConvertPS shader
- *****************************************************************************/
-
-BEGIN_GLOBAL_SHADER_PARAMETER_STRUCT(FRGB10toYUVv210ConvertUB, )
-SHADER_PARAMETER(FMatrix44f, ColorTransform)
-SHADER_PARAMETER(uint32, LinearToSrgb)
-SHADER_PARAMETER(float, OnePixelDeltaX)
-SHADER_PARAMETER(float, PaddingScale)
-SHADER_PARAMETER_TEXTURE(Texture2D, Texture)
-SHADER_PARAMETER_SAMPLER(SamplerState, SamplerP)
-END_GLOBAL_SHADER_PARAMETER_STRUCT()
-
-IMPLEMENT_GLOBAL_SHADER_PARAMETER_STRUCT(FRGB10toYUVv210ConvertUB, "RGB10toYUVv210ConvertUB");
-IMPLEMENT_SHADER_TYPE(, FRGB10toYUVv210ConvertPS, TEXT("/Engine/Private/MediaShaders.usf"), TEXT("RGB10toYUVv210ConvertPS"), SF_Pixel);
-
-
-void FRGB10toYUVv210ConvertPS::SetParameters(FRHICommandList& CommandList, TRefCountPtr<FRHITexture2D> RGBATexture, const FMatrix& ColorTransform, const FVector& YUVOffset, bool LinearToSrgb)
-{
-	FRGB10toYUVv210ConvertUB UB;
-	{
-		UB.ColorTransform = (FMatrix44f)MediaShaders::CombineColorTransformAndOffset(ColorTransform, YUVOffset);
-		UB.SamplerP = TStaticSamplerState<SF_Point>::GetRHI();
-		UB.LinearToSrgb = LinearToSrgb;
-		UB.Texture = RGBATexture;
-
-		//Output texture will be based on a size dividable by 48 (i.e 1280 -> 1296) and divided by 6 (i.e 1296 / 6 = 216)
-		//To map output texture UVs, we get a scale from the source texture original size to the mapped output size
-		//And use the source texture size to get the pixel delta
-		const float PaddedResolution = float(uint32((RGBATexture->GetSizeX() + 47) / 48) * 48);
-		UB.OnePixelDeltaX = 1.0f / (float)RGBATexture->GetSizeX();
-		UB.PaddingScale = PaddedResolution / (float)RGBATexture->GetSizeX();
-	}
-
-	TUniformBufferRef<FRGB10toYUVv210ConvertUB> Data = TUniformBufferRef<FRGB10toYUVv210ConvertUB>::CreateUniformBufferImmediate(UB, UniformBuffer_SingleFrame);
-	SetUniformBufferParameter(CommandList, CommandList.GetBoundPixelShader(), GetUniformBufferParameter<FRGB10toYUVv210ConvertUB>(), Data);
-}
-
-
 /* FRGB8toY8ConvertPS shader
  *****************************************************************************/
 
@@ -641,7 +575,7 @@ IMPLEMENT_GLOBAL_SHADER_PARAMETER_STRUCT(FReadTextureExternalUB, "ReadTextureExt
 IMPLEMENT_SHADER_TYPE(, FReadTextureExternalPS, TEXT("/Engine/Private/MediaShaders.usf"), TEXT("ReadTextureExternalPS"), SF_Pixel);
 
 
-void FReadTextureExternalPS::SetParameters(FRHICommandList& CommandList, FTextureRHIRef TextureExt, FSamplerStateRHIRef SamplerState, const FLinearColor & ScaleRotation, const FLinearColor & Offset)
+void FReadTextureExternalPS::SetParameters(FRHICommandList& CommandList, FTextureRHIRef TextureExt, FSamplerStateRHIRef SamplerState, const FLinearColor& ScaleRotation, const FLinearColor& Offset)
 {
 	FReadTextureExternalUB UB;
 	{
@@ -655,27 +589,67 @@ void FReadTextureExternalPS::SetParameters(FRHICommandList& CommandList, FTextur
 	SetUniformBufferParameter(CommandList, CommandList.GetBoundPixelShader(), GetUniformBufferParameter<FReadTextureExternalUB>(), Data);
 }
 
+/* FRGB8toUYVY8ConvertPS shader
+ *****************************************************************************/
+
+IMPLEMENT_GLOBAL_SHADER(FRGB8toUYVY8ConvertPS, "/Engine/Private/MediaShaders.usf", "RGB8toUYVY8ConvertPS", SF_Pixel);
+
+FRGB8toUYVY8ConvertPS::FParameters* FRGB8toUYVY8ConvertPS::AllocateAndSetParameters(FRDGBuilder& GraphBuilder, FRDGTextureRef RGBATexture, const FMatrix& ColorTransform, const FVector& YUVOffset, bool bDoLinearToSrgb, FRDGTextureRef OutputTexture)
+{
+	FRGB8toUYVY8ConvertPS::FParameters* Parameters = GraphBuilder.AllocParameters<FRGB8toUYVY8ConvertPS::FParameters>();
+
+	Parameters->RGBToYUVConversion.InputTexture = RGBATexture;
+	Parameters->RGBToYUVConversion.InputSampler = TStaticSamplerState<SF_Point>::GetRHI();
+	Parameters->RGBToYUVConversion.ColorTransform = (FMatrix44f)MediaShaders::CombineColorTransformAndOffset(ColorTransform, YUVOffset);
+	Parameters->RGBToYUVConversion.DoLinearToSrgb = bDoLinearToSrgb;
+	Parameters->RGBToYUVConversion.OnePixelDeltaX = 1.0f / (float)RGBATexture->Desc.Extent.X;
+	Parameters->RenderTargets[0] = FRenderTargetBinding{ OutputTexture, ERenderTargetLoadAction::ENoAction };
+
+	return Parameters;
+}
+
 
-/* FModifyAlphaSwizzleRgbaUB shader
+/* FRGB10toYUVv210ConvertPS shader
  *****************************************************************************/
 
-BEGIN_GLOBAL_SHADER_PARAMETER_STRUCT(FModifyAlphaSwizzleRgbaUB, )
-SHADER_PARAMETER_TEXTURE(Texture2D, Texture)
-SHADER_PARAMETER_SAMPLER(SamplerState, SamplerP)
-END_GLOBAL_SHADER_PARAMETER_STRUCT()
+IMPLEMENT_GLOBAL_SHADER(FRGB10toYUVv210ConvertPS, "/Engine/Private/MediaShaders.usf", "RGB10toYUVv210ConvertPS", SF_Pixel);
 
-IMPLEMENT_GLOBAL_SHADER_PARAMETER_STRUCT(FModifyAlphaSwizzleRgbaUB, "SwizzleRgbUB");
-IMPLEMENT_SHADER_TYPE(, FModifyAlphaSwizzleRgbaPS, TEXT("/Engine/Private/MediaShaders.usf"), TEXT("SwizzleRgbPS"), SF_Pixel);
+FRGB10toYUVv210ConvertPS::FParameters* FRGB10toYUVv210ConvertPS::AllocateAndSetParameters(FRDGBuilder& GraphBuilder, FRDGTextureRef RGBATexture, const FMatrix& ColorTransform, const FVector& YUVOffset, bool bDoLinearToSrgb, FRDGTextureRef OutputTexture)
+{
+	FRGB10toYUVv210ConvertPS::FParameters* Parameters = GraphBuilder.AllocParameters<FRGB10toYUVv210ConvertPS::FParameters>();
+
+	Parameters->RGBToYUVConversion.InputTexture = RGBATexture;
+	Parameters->RGBToYUVConversion.InputSampler = TStaticSamplerState<SF_Point>::GetRHI();
+	Parameters->RGBToYUVConversion.ColorTransform = (FMatrix44f)MediaShaders::CombineColorTransformAndOffset(ColorTransform, YUVOffset);
+	Parameters->RGBToYUVConversion.DoLinearToSrgb = bDoLinearToSrgb;
+	Parameters->RGBToYUVConversion.OnePixelDeltaX = 1.0f / (float)RGBATexture->Desc.Extent.X;
+
+	//Output texture will be based on a size dividable by 48 (i.e 1280 -> 1296) and divided by 6 (i.e 1296 / 6 = 216)
+	//To map output texture UVs, we get a scale from the source texture original size to the mapped output size
+	//And use the source texture size to get the pixel delta
+	const float PaddedResolution = float(uint32((RGBATexture->Desc.Extent.X + 47) / 48) * 48);
+	Parameters->PaddingScale = PaddedResolution / (float)RGBATexture->Desc.Extent.X;
 
+	Parameters->RenderTargets[0] = FRenderTargetBinding{ OutputTexture, ERenderTargetLoadAction::ENoAction };
+
+	return Parameters;
+}
 
-void FModifyAlphaSwizzleRgbaPS::SetParameters(FRHICommandList& CommandList, TRefCountPtr<FRHITexture2D> TextureExt)
+
+/* FModifyAlphaSwizzleRgbaPS shader
+ *****************************************************************************/
+
+IMPLEMENT_GLOBAL_SHADER(FModifyAlphaSwizzleRgbaPS, "/Engine/Private/MediaShaders.usf", "SwizzleRgbPS", SF_Pixel);
+
+FModifyAlphaSwizzleRgbaPS::FParameters* FModifyAlphaSwizzleRgbaPS::AllocateAndSetParameters(FRDGBuilder& GraphBuilder, FRDGTextureRef RGBATexture, FRDGTextureRef OutputTexture)
 {
-	FModifyAlphaSwizzleRgbaUB UniformBuffer;
-	{
-		UniformBuffer.SamplerP = TStaticSamplerState<SF_Point>::GetRHI();
-		UniformBuffer.Texture = TextureExt;
-	}
+	FModifyAlphaSwizzleRgbaPS::FParameters* Parameters = GraphBuilder.AllocParameters<FModifyAlphaSwizzleRgbaPS::FParameters>();
+
+	Parameters->InputTexture = RGBATexture;
+	Parameters->InputSampler = TStaticSamplerState<SF_Point>::GetRHI();
+	Parameters->RenderTargets[0] = FRenderTargetBinding{ OutputTexture, ERenderTargetLoadAction::ENoAction };
 
-	TUniformBufferRef<FModifyAlphaSwizzleRgbaUB> Data = TUniformBufferRef<FModifyAlphaSwizzleRgbaUB>::CreateUniformBufferImmediate(UniformBuffer, UniformBuffer_SingleFrame);
-	SetUniformBufferParameter(CommandList, CommandList.GetBoundPixelShader(), GetUniformBufferParameter<FModifyAlphaSwizzleRgbaUB>(), Data);
+	return Parameters;
 }
+
+
diff --git a/Engine/Source/Runtime/RenderCore/Public/MediaShaders.h b/Engine/Source/Runtime/RenderCore/Public/MediaShaders.h
--- a/Engine/Source/Runtime/RenderCore/Public/MediaShaders.h
+++ b/Engine/Source/Runtime/RenderCore/Public/MediaShaders.h
@@ -3,11 +3,12 @@
 #pragma once
 
 #include "CoreMinimal.h"
-#include "RHI.h"
-#include "RenderResource.h"
-#include "Shader.h"
 #include "GlobalShader.h"
 #include "PipelineStateCache.h"
+#include "RenderResource.h"
+#include "RHI.h"
+#include "Shader.h"
+#include "ShaderParameterStruct.h"
 
 
 /** MAX number of conversion operations. Reflects MAX in EMediaCaptureConversionOperation */
@@ -539,14 +540,14 @@ public:
 
 
 /**
- * Pixel shader to convert RGB 8 bits to UYVY 8 bits
+ * Pixel shader to convert RGB 8 bits to Y 8 bits
  *
  * This shader expects a single texture in PF_B8G8R8A8 format.
  */
-class FRGB8toUYVY8ConvertPS
+class FRGB8toY8ConvertPS
 	: public FGlobalShader
 {
-	DECLARE_EXPORTED_SHADER_TYPE(FRGB8toUYVY8ConvertPS, Global, RENDERCORE_API);
+	DECLARE_EXPORTED_SHADER_TYPE(FRGB8toY8ConvertPS, Global, RENDERCORE_API);
 
 public:
 
@@ -555,25 +556,23 @@ public:
 		return IsFeatureLevelSupported(Parameters.Platform, ERHIFeatureLevel::ES3_1);
 	}
 
-	FRGB8toUYVY8ConvertPS() { }
+	FRGB8toY8ConvertPS() { }
 
-	FRGB8toUYVY8ConvertPS(const ShaderMetaType::CompiledShaderInitializerType& Initializer)
+	FRGB8toY8ConvertPS(const ShaderMetaType::CompiledShaderInitializerType& Initializer)
 		: FGlobalShader(Initializer)
 	{ }
 
-	RENDERCORE_API void SetParameters(FRHICommandList& RHICmdList, TRefCountPtr<FRHITexture2D> RGBATexture, const FMatrix& ColorTransform, const FVector& YUVOffset, bool LinearToSrgb);
+	RENDERCORE_API void SetParameters(FRHICommandList& RHICmdList, TRefCountPtr<FRHITexture2D> RGBATexture, const FVector4f& ColorTransform, bool LinearToSrgb);
 };
 
 
 /**
- * Pixel shader to convert RGB 10 bits to YUV v210
- *
- * This shader expects a single texture in PF_A2B10G10R10 format.
+ * Pixel shader to read from TextureExternal source
  */
-class FRGB10toYUVv210ConvertPS
+class FReadTextureExternalPS
 	: public FGlobalShader
 {
-	DECLARE_EXPORTED_SHADER_TYPE(FRGB10toYUVv210ConvertPS, Global, RENDERCORE_API);
+	DECLARE_EXPORTED_SHADER_TYPE(FReadTextureExternalPS, Global, RENDERCORE_API);
 
 public:
 
@@ -582,25 +581,40 @@ public:
 		return IsFeatureLevelSupported(Parameters.Platform, ERHIFeatureLevel::ES3_1);
 	}
 
-	FRGB10toYUVv210ConvertPS() { }
+	FReadTextureExternalPS() { }
 
-	FRGB10toYUVv210ConvertPS(const ShaderMetaType::CompiledShaderInitializerType& Initializer)
+	FReadTextureExternalPS(const ShaderMetaType::CompiledShaderInitializerType& Initializer)
 		: FGlobalShader(Initializer)
 	{ }
 
-	RENDERCORE_API void SetParameters(FRHICommandList& RHICmdList, TRefCountPtr<FRHITexture2D> RGBATexture, const FMatrix& ColorTransform, const FVector& YUVOffset, bool LinearToSrgb);
+	RENDERCORE_API void SetParameters(FRHICommandList& RHICmdList, FTextureRHIRef TextureExt, FSamplerStateRHIRef SamplerState, const FLinearColor & ScaleRotation, const FLinearColor & Offset);
 };
 
 
+/** Struct of common parameters used in media capture shaders to do RGB to YUV conversions */
+BEGIN_SHADER_PARAMETER_STRUCT(FRGBToYUVConversion, RENDERCORE_API)
+	SHADER_PARAMETER_RDG_TEXTURE(Texture2D, InputTexture) 
+	SHADER_PARAMETER_SAMPLER(SamplerState, InputSampler)
+	SHADER_PARAMETER(FMatrix44f, ColorTransform)
+	SHADER_PARAMETER(uint32, DoLinearToSrgb)
+	SHADER_PARAMETER(float, OnePixelDeltaX)
+END_SHADER_PARAMETER_STRUCT()
+
 /**
- * Pixel shader to convert RGB 8 bits to Y 8 bits
- *
- * This shader expects a single texture in PF_B8G8R8A8 format.
+ * Pixel shader to convert RGB 8 bits to UYVY 8 bits
  */
-class FRGB8toY8ConvertPS
-	: public FGlobalShader
+	class RENDERCORE_API FRGB8toUYVY8ConvertPS : public FGlobalShader
 {
-	DECLARE_EXPORTED_SHADER_TYPE(FRGB8toY8ConvertPS, Global, RENDERCORE_API);
+public:
+
+	DECLARE_GLOBAL_SHADER(FRGB8toUYVY8ConvertPS);
+
+	SHADER_USE_PARAMETER_STRUCT(FRGB8toUYVY8ConvertPS, FGlobalShader);
+
+	BEGIN_SHADER_PARAMETER_STRUCT(FParameters, )
+		SHADER_PARAMETER_STRUCT_INCLUDE(FRGBToYUVConversion, RGBToYUVConversion)
+		RENDER_TARGET_BINDING_SLOTS()
+	END_SHADER_PARAMETER_STRUCT()
 
 public:
 
@@ -609,23 +623,32 @@ public:
 		return IsFeatureLevelSupported(Parameters.Platform, ERHIFeatureLevel::ES3_1);
 	}
 
-	FRGB8toY8ConvertPS() { }
-
-	FRGB8toY8ConvertPS(const ShaderMetaType::CompiledShaderInitializerType& Initializer)
-		: FGlobalShader(Initializer)
-	{ }
+	UE_DEPRECATED(5.1, "SetParameters has been deprecated while moving to a RDG based pipeline. Please use AllocateAndSetParameters instead.")
+	void SetParameters(FRHICommandList& RHICmdList, TRefCountPtr<FRHITexture2D> RGBATexture, const FMatrix& ColorTransform, const FVector& YUVOffset, bool LinearToSrgb)
+	{
+	}
 
-	RENDERCORE_API void SetParameters(FRHICommandList& RHICmdList, TRefCountPtr<FRHITexture2D> RGBATexture, const FVector4f& ColorTransform, bool LinearToSrgb);
+	/** Allocates and setup shader parameter in the incoming graph builder */
+	FRGB8toUYVY8ConvertPS::FParameters* AllocateAndSetParameters(FRDGBuilder& GraphBuilder, FRDGTextureRef RGBATexture, const FMatrix& ColorTransform, const FVector& YUVOffset, bool bDoLinearToSrgb, FRDGTextureRef OutputTexture);
 };
 
 
 /**
- * Pixel shader to read from TextureExternal source
+ * Pixel shader to convert RGB 10 bits to YUV v210
  */
-class FReadTextureExternalPS
-	: public FGlobalShader
+	class RENDERCORE_API FRGB10toYUVv210ConvertPS : public FGlobalShader
 {
-	DECLARE_EXPORTED_SHADER_TYPE(FReadTextureExternalPS, Global, RENDERCORE_API);
+public:
+
+	DECLARE_GLOBAL_SHADER(FRGB10toYUVv210ConvertPS);
+
+	SHADER_USE_PARAMETER_STRUCT(FRGB10toYUVv210ConvertPS, FGlobalShader);
+
+	BEGIN_SHADER_PARAMETER_STRUCT(FParameters, )
+		SHADER_PARAMETER_STRUCT_INCLUDE(FRGBToYUVConversion, RGBToYUVConversion)
+		SHADER_PARAMETER(float, PaddingScale)
+		RENDER_TARGET_BINDING_SLOTS()
+	END_SHADER_PARAMETER_STRUCT()
 
 public:
 
@@ -634,47 +657,52 @@ public:
 		return IsFeatureLevelSupported(Parameters.Platform, ERHIFeatureLevel::ES3_1);
 	}
 
-	FReadTextureExternalPS() { }
-
-	FReadTextureExternalPS(const ShaderMetaType::CompiledShaderInitializerType& Initializer)
-		: FGlobalShader(Initializer)
-	{ }
+	UE_DEPRECATED(5.1, "SetParameters has been deprecated while moving to a RDG based pipeline. Please use AllocateAndSetParameters instead.")
+	void SetParameters(FRHICommandList& RHICmdList, TRefCountPtr<FRHITexture2D> RGBATexture, const FMatrix& ColorTransform, const FVector& YUVOffset, bool LinearToSrgb)
+	{
+	}
 
-	RENDERCORE_API void SetParameters(FRHICommandList& RHICmdList, FTextureRHIRef TextureExt, FSamplerStateRHIRef SamplerState, const FLinearColor & ScaleRotation, const FLinearColor & Offset);
+	/** Allocates and setup shader parameter in the incoming graph builder */
+	FRGB10toYUVv210ConvertPS::FParameters* AllocateAndSetParameters(FRDGBuilder& GraphBuilder, FRDGTextureRef RGBATexture, const FMatrix& ColorTransform, const FVector& YUVOffset, bool bDoLinearToSrgb, FRDGTextureRef OutputTexture);
 };
 
-
 /**
  * Pixel shader to swizzle R G B A components, set alpha to 1 or inverts alpha
  *
  * General conversion shader that is only used to swizzle shaders that do not require any color conversion. RGB to BGR, RGB10A2 to RGBA8 etc
  */
-class FModifyAlphaSwizzleRgbaPS
-	: public FGlobalShader
+class RENDERCORE_API FModifyAlphaSwizzleRgbaPS	: public FGlobalShader
 {
 public:
-	DECLARE_EXPORTED_SHADER_TYPE(FModifyAlphaSwizzleRgbaPS, Global, RENDERCORE_API);
+
+	DECLARE_GLOBAL_SHADER(FModifyAlphaSwizzleRgbaPS);
+
+	SHADER_USE_PARAMETER_STRUCT(FModifyAlphaSwizzleRgbaPS, FGlobalShader);
 
 	class FConversionOp : SHADER_PERMUTATION_INT("CONVERSION_OP", NUM_MEDIA_SHADERS_CONVERSION_OP);
 	using FPermutationDomain = TShaderPermutationDomain<FConversionOp>;
 
+	BEGIN_SHADER_PARAMETER_STRUCT(FParameters, )
+		SHADER_PARAMETER_RDG_TEXTURE(Texture2D, InputTexture) 
+		SHADER_PARAMETER_SAMPLER(SamplerState, InputSampler)
+		RENDER_TARGET_BINDING_SLOTS()
+	END_SHADER_PARAMETER_STRUCT()
+
+public:
 	static bool ShouldCompilePermutation(const FGlobalShaderPermutationParameters& Parameters)
 	{
 		FPermutationDomain PermutationVector(Parameters.PermutationId);
 		return IsFeatureLevelSupported(Parameters.Platform, ERHIFeatureLevel::ES3_1);
 	}
 
-	static void ModifyCompilationEnvironment(const FGlobalShaderPermutationParameters& Parameters, FShaderCompilerEnvironment& OutEnvironment)
+	UE_DEPRECATED(5.1, "SetParameters has been deprecated while moving to a RDG based pipeline. Please use AllocateAndSetParameters instead.")
+	void SetParameters(FRHICommandList& RHICmdList, TRefCountPtr<FRHITexture2D> RGBATexture)
 	{
-		FGlobalShader::ModifyCompilationEnvironment(Parameters, OutEnvironment);
 	}
 
-	FModifyAlphaSwizzleRgbaPS() = default;
+	/** Allocates and setup shader parameter in the incoming graph builder */
+	FModifyAlphaSwizzleRgbaPS::FParameters* AllocateAndSetParameters(FRDGBuilder& GraphBuilder, FRDGTextureRef RGBATexture, FRDGTextureRef OutputTexture);
+};
 
-	FModifyAlphaSwizzleRgbaPS(const ShaderMetaType::CompiledShaderInitializerType& Initializer)
-		: FGlobalShader(Initializer)
-	{ }
 
-	RENDERCORE_API void SetParameters(FRHICommandList& RHICmdList, TRefCountPtr<FRHITexture2D> RGBATexture);
-};
 
