- Setting bytes per row back to what it was before media capture refactor. Still missing a bytes per pixel to be fully fixed.

diff --git a/Engine/Plugins/Media/MediaIOFramework/Source/MediaIOCore/Private/MediaCapture.cpp b/Engine/Plugins/Media/MediaIOFramework/Source/MediaIOCore/Private/MediaCapture.cpp
--- a/Engine/Plugins/Media/MediaIOFramework/Source/MediaIOCore/Private/MediaCapture.cpp
+++ b/Engine/Plugins/Media/MediaIOFramework/Source/MediaIOCore/Private/MediaCapture.cpp
@@ -1424,8 +1424,6 @@ void UMediaCapture::Capture_RenderThread(FRHICommandListImmediate& RHICmdList,
 #endif
 			// Lock & read
 			void* ColorDataBuffer = nullptr;
-			int32 Width = 0;
-			int32 Height = 0;
 			int32 RowPitchInPixels = 0;
 
 			// Texture readback does no verification if it's ready. It is assumed it will stall if it's not
@@ -1440,7 +1438,7 @@ void UMediaCapture::Capture_RenderThread(FRHICommandListImmediate& RHICmdList,
 
 				// The Width/Height of the surface may be different then the DesiredOutputSize : Some underlying implementations enforce a specific stride, therefore
 				// there may be padding at the end of each row.
-				InMediaCapture->OnFrameCaptured_RenderingThread(ReadyFrame->CaptureBaseData, ReadyFrame->UserData, ColorDataBuffer, InMediaCapture->DesiredOutputSize.X, InMediaCapture->DesiredOutputSize.Y, Width * 4);
+				InMediaCapture->OnFrameCaptured_RenderingThread(ReadyFrame->CaptureBaseData, ReadyFrame->UserData, ColorDataBuffer, InMediaCapture->DesiredOutputSize.X, InMediaCapture->DesiredOutputSize.Y, RowPitchInPixels * 4);
 			}
 			ReadyFrame->bReadbackRequested = false;
 
