VulkanRHI:  Avoid GetSelectedDynamicRHIModuleName in module startup by always adding necessary Vulkan extensions (it will just get ignored if Vulkan isn't used).

#rnx

diff --git a/Engine/Plugins/Media/MediaIOFramework/Source/GPUTextureTransfer/Private/GPUTextureTransferModule.cpp b/Engine/Plugins/Media/MediaIOFramework/Source/GPUTextureTransfer/Private/GPUTextureTransferModule.cpp
--- a/Engine/Plugins/Media/MediaIOFramework/Source/GPUTextureTransfer/Private/GPUTextureTransferModule.cpp
+++ b/Engine/Plugins/Media/MediaIOFramework/Source/GPUTextureTransfer/Private/GPUTextureTransferModule.cpp
@@ -49,17 +49,7 @@ void FGPUTextureTransferModule::StartupModule()
 	{
 		if (LoadGPUDirectBinary())
 		{
-#if DVP_SUPPORTED_PLATFORM		
-			const TCHAR* DynamicRHIModuleName = GetSelectedDynamicRHIModuleName(false);
-#elif PLATFORM_LINUX
-			const TCHAR* DynamicRHIModuleName = TEXT("VulkanRHI");
-#else
-			const TCHAR* DynamicRHIModuleName = TEXT("");
-			ensure(false);
-#endif
-
-			// We cannot use GDynmicRHI here because it hasn't been assigned yet.
-			if (TEXT("VulkanRHI") == FString(DynamicRHIModuleName))
+			// Always provide the necessary Vulkan extensions (it will just get ignored if a different RHI is used)
 			{
 #if DVP_SUPPORTED_PLATFORM
 				const TArray<const ANSICHAR*> ExtentionsToAdd{ VK_KHR_EXTERNAL_MEMORY_EXTENSION_NAME, VK_KHR_EXTERNAL_MEMORY_WIN32_EXTENSION_NAME, VK_KHR_EXTERNAL_SEMAPHORE_WIN32_EXTENSION_NAME, VK_KHR_SURFACE_EXTENSION_NAME, VK_KHR_WIN32_SURFACE_EXTENSION_NAME };
