Added 'GetDeviceProviders' to be able to query all the currently registered Device Providers

diff --git a/Engine/Plugins/Media/MediaIOFramework/Source/MediaIOCore/Private/MediaIOCoreModule.cpp b/Engine/Plugins/Media/MediaIOFramework/Source/MediaIOCore/Private/MediaIOCoreModule.cpp
--- a/Engine/Plugins/Media/MediaIOFramework/Source/MediaIOCore/Private/MediaIOCoreModule.cpp
+++ b/Engine/Plugins/Media/MediaIOFramework/Source/MediaIOCore/Private/MediaIOCoreModule.cpp
@@ -40,6 +40,11 @@ public:
 		return nullptr;
 	}
 
+	virtual TConstArrayView<IMediaIOCoreDeviceProvider*> GetDeviceProviders() const override
+	{
+		return DeviceProviders;
+	}
+
 private:
 	TArray<IMediaIOCoreDeviceProvider*> DeviceProviders;
 };
diff --git a/Engine/Plugins/Media/MediaIOFramework/Source/MediaIOCore/Public/IMediaIOCoreModule.h b/Engine/Plugins/Media/MediaIOFramework/Source/MediaIOCore/Public/IMediaIOCoreModule.h
--- a/Engine/Plugins/Media/MediaIOFramework/Source/MediaIOCore/Public/IMediaIOCoreModule.h
+++ b/Engine/Plugins/Media/MediaIOFramework/Source/MediaIOCore/Public/IMediaIOCoreModule.h
@@ -21,4 +21,5 @@ public:
 	virtual void RegisterDeviceProvider(IMediaIOCoreDeviceProvider* InProvider) = 0;
 	virtual void UnregisterDeviceProvider(IMediaIOCoreDeviceProvider* InProvider) = 0;
 	virtual IMediaIOCoreDeviceProvider* GetDeviceProvider(FName InProviderName) = 0;
+	virtual TConstArrayView<IMediaIOCoreDeviceProvider*> GetDeviceProviders() const = 0;
 };
