- Fixing crash when using file media capture with resolution that needs padding
- Fix float FileMedia output expected stride
- Fix log spam when starting file media capture

diff --git a/Engine/Plugins/Media/MediaIOFramework/Source/MediaIOCore/MediaIOCore.Build.cs b/Engine/Plugins/Media/MediaIOFramework/Source/MediaIOCore/MediaIOCore.Build.cs
--- a/Engine/Plugins/Media/MediaIOFramework/Source/MediaIOCore/MediaIOCore.Build.cs
+++ b/Engine/Plugins/Media/MediaIOFramework/Source/MediaIOCore/MediaIOCore.Build.cs
@@ -41,6 +41,7 @@ namespace UnrealBuildTool.Rules
                     "AudioMixer",
 	                "AudioMixerCore",
 					"GPUTextureTransfer",
+					"ImageWrapper",
                     "RenderCore",
 					"Renderer",
                     "SignalProcessing", 
diff --git a/Engine/Plugins/Media/MediaIOFramework/Source/MediaIOCore/Private/FileMediaCapture.cpp b/Engine/Plugins/Media/MediaIOFramework/Source/MediaIOCore/Private/FileMediaCapture.cpp
--- a/Engine/Plugins/Media/MediaIOFramework/Source/MediaIOCore/Private/FileMediaCapture.cpp
+++ b/Engine/Plugins/Media/MediaIOFramework/Source/MediaIOCore/Private/FileMediaCapture.cpp
@@ -3,8 +3,10 @@
 #include "FileMediaCapture.h"
 
 #include "FileMediaOutput.h"
+#include "IImageWrapperModule.h"
 #include "ImageWriteQueue.h"
 #include "ImageWriteTask.h"
+#include "MediaIOCoreModule.h"
 #include "Misc/Paths.h"
 #include "Modules/ModuleManager.h"
 
@@ -74,9 +76,16 @@ void UFileMediaCapture::OnFrameCaptured_RenderingThread(const FCaptureBaseData&
 		return;
 	}
 
+	IImageWrapperModule* ImageModulePtr = FModuleManager::Get().GetModulePtr<IImageWrapperModule>("ImageWrapper");
+	if (ImageModulePtr == nullptr)
+	{
+		SetState(EMediaCaptureState::Error);
+		return;
+	}
+
 	TUniquePtr<FImageWriteTask> ImageTask = MakeUnique<FImageWriteTask>();
 	ImageTask->Format = ImageFormat;
-	ImageTask->Filename = FString::Printf(TEXT("%s%05d"), *BaseFilePathName, InBaseData.SourceFrameNumber);
+	ImageTask->Filename = FString::Printf(TEXT("%s%05d.%s"), *BaseFilePathName, InBaseData.SourceFrameNumber, ImageModulePtr->GetExtension(ImageFormat));
 	ImageTask->bOverwriteFile = bOverwriteFile;
 	ImageTask->CompressionQuality = CompressionQuality;
 	ImageTask->OnCompleted = OnCompleteWrapper;
@@ -85,7 +94,12 @@ void UFileMediaCapture::OnFrameCaptured_RenderingThread(const FCaptureBaseData&
 	if (PixelFormat == PF_B8G8R8A8)
 	{
 		// We only support tightly packed rows without padding
-		check((BytesPerRow == 0) || (BytesPerRow == (Width * 4)));
+		if ((BytesPerRow != 0) && (BytesPerRow != (Width * 4)))
+		{
+			UE_LOG(LogMediaIOCore, Error, TEXT("File media capture only supports tightly packed rows. Expected stride: %d. Received stride: %d. It might also mean that output resolution is too small."), Width * 4, BytesPerRow);
+			SetState(EMediaCaptureState::Error);
+			return;
+		}
 
 		TUniquePtr<TImagePixelData<FColor>> PixelData = MakeUnique<TImagePixelData<FColor>>(FIntPoint(Width, Height),
 			TArray<FColor, FDefaultAllocator64>(reinterpret_cast<FColor*>(InBuffer), Width * Height));
@@ -94,7 +108,12 @@ void UFileMediaCapture::OnFrameCaptured_RenderingThread(const FCaptureBaseData&
 	else if (PixelFormat == PF_FloatRGBA)
 	{
 		// We only support tightly packed rows without padding
-		check((BytesPerRow == 0) || (BytesPerRow == (Width * 16)));
+		if ((BytesPerRow != 0) && (BytesPerRow != (Width * 8)))
+		{
+			UE_LOG(LogMediaIOCore, Error, TEXT("File media capture only supports tightly packed rows. Expected stride: %d. Received stride: %d. It might also mean that output resolution is too small."), Width * 8, BytesPerRow);
+			SetState(EMediaCaptureState::Error);
+			return;
+		}
 
 		TUniquePtr<TImagePixelData<FFloat16Color>> PixelData = MakeUnique<TImagePixelData<FFloat16Color>>(FIntPoint(Width, Height), 
 			TArray<FFloat16Color, FDefaultAllocator64>(reinterpret_cast<FFloat16Color*>(InBuffer), Width * Height));
diff --git a/Engine/Plugins/Media/MediaIOFramework/Source/MediaIOCore/Private/MediaCapture.cpp b/Engine/Plugins/Media/MediaIOFramework/Source/MediaIOCore/Private/MediaCapture.cpp
--- a/Engine/Plugins/Media/MediaIOFramework/Source/MediaIOCore/Private/MediaCapture.cpp
+++ b/Engine/Plugins/Media/MediaIOFramework/Source/MediaIOCore/Private/MediaCapture.cpp
@@ -86,6 +86,49 @@
 	PF_R8,
 };
 
+/* namespace MediaCaptureDetails definition
+*****************************************************************************/
+
+namespace MediaCaptureDetails
+{
+	bool FindSceneViewportAndLevel(TSharedPtr<FSceneViewport>& OutSceneViewport);
+
+	//Validation for the source of a capture
+	bool ValidateSceneViewport(const TSharedPtr<FSceneViewport>& SceneViewport, const FMediaCaptureOptions& CaptureOption, const FIntPoint& DesiredSize, const EPixelFormat DesiredPixelFormat, const bool bCurrentlyCapturing);
+	bool ValidateTextureRenderTarget2D(const UTextureRenderTarget2D* RenderTarget, const FMediaCaptureOptions& CaptureOption, const FIntPoint& DesiredSize, const EPixelFormat DesiredPixelFormat, const bool bCurrentlyCapturing);
+
+	//Validation that there is a capture
+	bool ValidateIsCapturing(const UMediaCapture& CaptureToBeValidated);
+
+	void ShowSlateNotification();
+
+	/** Returns bytes per pixel based on pixel format */
+	int32 GetBytesPerPixel(EPixelFormat InPixelFormat);
+
+	static const FName LevelEditorName(TEXT("LevelEditor"));
+}
+
+#if WITH_EDITOR
+namespace MediaCaptureAnalytics
+{
+	/**
+	 * @EventName MediaFramework.CaptureStarted
+	 * @Trigger Triggered when a capture of the viewport or render target is started.
+	 * @Type Client
+	 * @Owner MediaIO Team
+	 */
+	void SendCaptureEvent(const FString& CaptureType)
+	{
+		if (FEngineAnalytics::IsAvailable())
+		{
+			TArray<FAnalyticsEventAttribute> EventAttributes;
+			EventAttributes.Add(FAnalyticsEventAttribute(TEXT("CaptureType"), CaptureType));
+			FEngineAnalytics::GetProvider().RecordEvent(TEXT("MediaFramework.CaptureStarted"), EventAttributes);
+		}
+	}
+}
+#endif
+
 namespace UE::MediaCaptureData
 {
 	class FCaptureFrame
@@ -106,7 +149,7 @@
 		virtual bool IsBufferResource() const = 0;
 
 		/** Locks the readback resource and returns a pointer to access data from system memory */
-		virtual void* Lock(FRHICommandListImmediate& RHICmdList, FRHIGPUMask GPUMask, int32& OutWidth) = 0;
+		virtual void* Lock(FRHICommandListImmediate& RHICmdList, FRHIGPUMask GPUMask, int32& OutRowStride) = 0;
 
 		/** Unlocks the readback resource */
 		virtual void Unlock() = 0;
@@ -141,14 +184,17 @@
 			return false;
 		}
 
-		virtual void* Lock(FRHICommandListImmediate& RHICmdList, FRHIGPUMask GPUMask, int32& OutWidth) override
+		virtual void* Lock(FRHICommandListImmediate& RHICmdList, FRHIGPUMask GPUMask, int32& OutRowStride) override
 		{
 			if (ReadbackTexture->IsReady(GPUMask) == false)
 			{
 				UE_LOG(LogMediaIOCore, Verbose, TEXT("Fenced for texture readback was not ready"));
 			}
 
-			return ReadbackTexture->Lock(OutWidth);
+			int32 ReadbackWidth;
+			void* ReadbackPointer = ReadbackTexture->Lock(ReadbackWidth);
+			OutRowStride = ReadbackWidth * MediaCaptureDetails::GetBytesPerPixel(RenderTarget->GetDesc().Format);
+			return ReadbackPointer;
 		}
 
 		virtual void Unlock() override
@@ -208,7 +254,7 @@
 			return true;
 		}
 
-		virtual void* Lock(FRHICommandListImmediate& RHICmdList, FRHIGPUMask GPUMask, int32& OutWidth) override
+		virtual void* Lock(FRHICommandListImmediate& RHICmdList, FRHIGPUMask GPUMask, int32& OutRowStride) override
 		{
 			if (ReadbackBuffer->IsReady(GPUMask) == false)
 			{
@@ -216,7 +262,7 @@
 				RHICmdList.BlockUntilGPUIdle();
 			}
 
-			OutWidth = Buffer->GetRHI()->GetStride();
+			OutRowStride = Buffer->GetRHI()->GetStride();
 			return ReadbackBuffer->Lock(Buffer->GetRHI()->GetSize());
 		}
 
@@ -587,48 +633,6 @@
 }
 
 
-
-/* namespace MediaCaptureDetails definition
-*****************************************************************************/
-
-namespace MediaCaptureDetails
-{
-	bool FindSceneViewportAndLevel(TSharedPtr<FSceneViewport>& OutSceneViewport);
-
-	//Validation for the source of a capture
-	bool ValidateSceneViewport(const TSharedPtr<FSceneViewport>& SceneViewport, const FMediaCaptureOptions& CaptureOption, const FIntPoint& DesiredSize, const EPixelFormat DesiredPixelFormat, const bool bCurrentlyCapturing);
-	bool ValidateTextureRenderTarget2D(const UTextureRenderTarget2D* RenderTarget, const FMediaCaptureOptions& CaptureOption, const FIntPoint& DesiredSize, const EPixelFormat DesiredPixelFormat, const bool bCurrentlyCapturing);
-
-	//Validation that there is a capture
-	bool ValidateIsCapturing(const UMediaCapture& CaptureToBeValidated);
-
-	void ShowSlateNotification();
-
-	static const FName LevelEditorName(TEXT("LevelEditor"));
-}
-
-#if WITH_EDITOR
-namespace MediaCaptureAnalytics
-{
-	/**
-	 * @EventName MediaFramework.CaptureStarted
-	 * @Trigger Triggered when a capture of the viewport or render target is started.
-	 * @Type Client
-	 * @Owner MediaIO Team
-	 */
-	void SendCaptureEvent(const FString& CaptureType)
-	{
-		if (FEngineAnalytics::IsAvailable())
-		{
-			TArray<FAnalyticsEventAttribute> EventAttributes;
-			EventAttributes.Add(FAnalyticsEventAttribute(TEXT("CaptureType"), CaptureType));
-			FEngineAnalytics::GetProvider().RecordEvent(TEXT("MediaFramework.CaptureStarted"), EventAttributes);
-		}
-	}
-}
-#endif
-
-
 /* UMediaCapture::FCaptureBaseData
 *****************************************************************************/
 UMediaCapture::FCaptureBaseData::FCaptureBaseData()
@@ -1424,13 +1428,13 @@
 #endif
 			// Lock & read
 			void* ColorDataBuffer = nullptr;
-			int32 RowPitchInPixels = 0;
+			int32 RowStride = 0;
 
 			// Texture readback does no verification if it's ready. It is assumed it will stall if it's not
 			// Buffer readback will verify and block GPU until idle in case it's not ready otherwise it ensures
 			{
 				SCOPE_CYCLE_COUNTER(STAT_MediaCapture_RenderThread_LockResource);
-				ColorDataBuffer = ReadyFrame->Lock(RHICmdList, GPUMask, RowPitchInPixels);
+				ColorDataBuffer = ReadyFrame->Lock(RHICmdList, GPUMask, RowStride);
 			}
 
 			{
@@ -1438,7 +1442,7 @@
 
 				// The Width/Height of the surface may be different then the DesiredOutputSize : Some underlying implementations enforce a specific stride, therefore
 				// there may be padding at the end of each row.
-				InMediaCapture->OnFrameCaptured_RenderingThread(ReadyFrame->CaptureBaseData, ReadyFrame->UserData, ColorDataBuffer, InMediaCapture->DesiredOutputSize.X, InMediaCapture->DesiredOutputSize.Y, RowPitchInPixels * 4);
+				InMediaCapture->OnFrameCaptured_RenderingThread(ReadyFrame->CaptureBaseData, ReadyFrame->UserData, ColorDataBuffer, InMediaCapture->DesiredOutputSize.X, InMediaCapture->DesiredOutputSize.Y, RowStride);
 			}
 			ReadyFrame->bReadbackRequested = false;
 
@@ -1651,6 +1655,75 @@
 		}
 #endif // WITH_EDITOR
 	}
+
+	int32 GetBytesPerPixel(EPixelFormat InPixelFormat)
+	{
+		//We can capture viewports and render targets. Possible pixel format is limited by that
+		switch (InPixelFormat)
+		{
+
+		case PF_A8:
+		case PF_R8_UINT:
+		case PF_R8_SINT:
+		case PF_G8:
+		{
+			return 1;
+		}
+		case PF_R16_UINT:
+		case PF_R16_SINT:
+		case PF_R5G6B5_UNORM:
+		case PF_R8G8:
+		case PF_R16F:
+		case PF_R16F_FILTER:
+		case PF_V8U8:
+		case PF_R8G8_UINT:
+		case PF_B5G5R5A1_UNORM:
+		{
+			return 2;
+		}
+		case PF_R32_UINT:
+		case PF_R32_SINT:
+		case PF_R8G8B8A8:
+		case PF_A8R8G8B8:
+		case PF_FloatR11G11B10:
+		case PF_A2B10G10R10:
+		case PF_G16R16:
+		case PF_G16R16F:
+		case PF_G16R16F_FILTER:
+		case PF_R32_FLOAT:
+		case PF_R16G16_UINT:
+		case PF_R8G8B8A8_UINT:
+		case PF_R8G8B8A8_SNORM:
+		case PF_B8G8R8A8:
+		case PF_G16R16_SNORM:
+		case PF_FloatRGB: //Equivalent to R11G11B10
+		{
+			return 4;
+		}
+		case PF_R16G16B16A16_UINT:
+		case PF_R16G16B16A16_SINT:
+		case PF_A16B16G16R16:
+		case PF_G32R32F:
+		case PF_R16G16B16A16_UNORM:
+		case PF_R16G16B16A16_SNORM:
+		case PF_R32G32_UINT:
+		case PF_R64_UINT:
+		case PF_FloatRGBA: //Equivalent to R16G16B16A16
+		{
+			return 8;
+		}
+		case PF_A32B32G32R32F:
+		case PF_R32G32B32A32_UINT:
+		{
+			return 16;
+		}
+		default:
+		{
+			ensureMsgf(false, TEXT("MediaCapture - Pixel format (%d) not handled. Invalid bytes per pixel returned."), InPixelFormat);
+			return 0;
+		}
+		}
+	}
 }
 
 #undef LOCTEXT_NAMESPACE
