Fix wrong text for input timecode customization

diff --git a/Engine/Plugins/Media/MediaIOFramework/Source/MediaIOCore/Private/MediaIOCoreDefinitions.cpp b/Engine/Plugins/Media/MediaIOFramework/Source/MediaIOCore/Private/MediaIOCoreDefinitions.cpp
--- a/Engine/Plugins/Media/MediaIOFramework/Source/MediaIOCore/Private/MediaIOCoreDefinitions.cpp
+++ b/Engine/Plugins/Media/MediaIOFramework/Source/MediaIOCore/Private/MediaIOCoreDefinitions.cpp
@@ -141,7 +141,7 @@ FText FMediaIOMode::GetModeName() const
 		}
 		return FMediaIOCommonDisplayModes::GetMediaIOCommonDisplayModeInfoName(Resolution.X, Resolution.Y, FieldFrameRate, Standard);
 	}
-	return LOCTEXT("Invalid", "<Invalid>");
+	return LOCTEXT("InvalidMode", "<Invalid>");
 }
 
 
@@ -295,19 +295,19 @@ FText FMediaIOVideoTimecodeConfiguration::ToText(bool bAutoDetected) const
 {
 	if (bAutoDetected)
 	{
-		const FText ConfigurationText = FText::Format(LOCTEXT("FMediaIOConfigurationToText", "{0} [device{1}/{2}{3}]")
+		const FText ConfigurationText = FText::Format(LOCTEXT("FMediaIOAutoDetectConfigurationToText", "{0} [device{1}/{2}{3}]")
         			, FText::FromName(MediaConfiguration.MediaConnection.Device.DeviceName)
         			, FText::AsNumber(MediaConfiguration.MediaConnection.Device.DeviceIdentifier)
         			, IMediaIOCoreDeviceProvider::GetTransportName(MediaConfiguration.MediaConnection.TransportType, MediaConfiguration.MediaConnection.QuadTransportType)
         			, FText::AsNumber(MediaConfiguration.MediaConnection.PortIdentifier)
         		);
 		
-		return FText::Format(LOCTEXT("MediaTimecodeConfigurationRefText", "{0} / Auto"), ConfigurationText);
+		return FText::Format(LOCTEXT("MediaTimecodeAutoDetectConfigurationRefText", "{0} / Auto"), ConfigurationText);
 	}
 	
 	if (IsValid())
 	{
-		FText Timecode = LOCTEXT("Invalid", "<Invalid>");
+		FText Timecode = LOCTEXT("InvalidTimecode", "<Invalid>");
 		switch(TimecodeFormat)
 		{
 		case EMediaIOAutoDetectableTimecodeFormat::LTC:
@@ -332,7 +332,7 @@ FText FMediaIOVideoTimecodeConfiguration::ToText(bool bAutoDetected) const
 			, Timecode
 		);
 	}
-	return LOCTEXT("Invalid", "<Invalid>");
+	return LOCTEXT("InvalidConfiguration", "<Invalid>");
 }
 
 
diff --git a/Engine/Plugins/Media/MediaIOFramework/Source/MediaIOEditor/Private/Customizations/MediaIOVideoTimecodeConfigurationCustomization.cpp b/Engine/Plugins/Media/MediaIOFramework/Source/MediaIOEditor/Private/Customizations/MediaIOVideoTimecodeConfigurationCustomization.cpp
--- a/Engine/Plugins/Media/MediaIOFramework/Source/MediaIOEditor/Private/Customizations/MediaIOVideoTimecodeConfigurationCustomization.cpp
+++ b/Engine/Plugins/Media/MediaIOFramework/Source/MediaIOEditor/Private/Customizations/MediaIOVideoTimecodeConfigurationCustomization.cpp
@@ -69,7 +69,7 @@ TAttribute<FText> FMediaIOVideoTimecodeConfigurationCustomization::GetContentTex
 	FMediaIOVideoTimecodeConfiguration* Value = GetPropertyValueFromPropertyHandle<FMediaIOVideoTimecodeConfiguration>();
 	return MakeAttributeLambda([=]
 	{
-		const bool bIsAutoDetected = GetEnforceCheckboxState() == ECheckBoxState::Unchecked; 
+		const bool bIsAutoDetected = IsAutoDetected();
 		return Value->ToText(bIsAutoDetected);
 	});
 }
