Fix assertion in mediaiocoredefinitions

diff --git a/Engine/Plugins/Media/MediaIOFramework/Source/MediaIOCore/Public/MediaIOCoreDefinitions.h b/Engine/Plugins/Media/MediaIOFramework/Source/MediaIOCore/Public/MediaIOCoreDefinitions.h
--- a/Engine/Plugins/Media/MediaIOFramework/Source/MediaIOCore/Public/MediaIOCoreDefinitions.h
+++ b/Engine/Plugins/Media/MediaIOFramework/Source/MediaIOCore/Public/MediaIOCoreDefinitions.h
@@ -92,6 +92,8 @@ namespace UE::MediaIO
 			return EMediaIOTimecodeFormat::LTC;
 		case EMediaIOAutoDetectableTimecodeFormat::VITC:
 			return EMediaIOTimecodeFormat::VITC;
+		case EMediaIOAutoDetectableTimecodeFormat::None:
+			return EMediaIOTimecodeFormat::None;
 		default:
 			ensureMsgf(false, TEXT("Unknown timecode format encountered!"));
 			return EMediaIOTimecodeFormat::None;
