Add missing include in mediaIO to fix CIS

diff --git a/Engine/Plugins/Media/MediaIOFramework/Source/MediaIOEditor/Private/Customizations/MediaIOConfigurationCustomization.cpp b/Engine/Plugins/Media/MediaIOFramework/Source/MediaIOEditor/Private/Customizations/MediaIOConfigurationCustomization.cpp
--- a/Engine/Plugins/Media/MediaIOFramework/Source/MediaIOEditor/Private/Customizations/MediaIOConfigurationCustomization.cpp
+++ b/Engine/Plugins/Media/MediaIOFramework/Source/MediaIOEditor/Private/Customizations/MediaIOConfigurationCustomization.cpp
@@ -12,6 +12,7 @@
 #include "IMediaIOCoreModule.h"
 #include "MediaIOPermutationsSelectorBuilder.h"
 #include "ObjectEditorUtils.h"
+#include "ScopedTransaction.h"
 #include "TimeSynchronizableMediaSource.h"
 #include "Widgets/Input/SCheckBox.h"
 #include "Widgets/Input/SComboButton.h"
