MediaCapture: Removed unused variable.

diff --git a/Engine/Plugins/Media/MediaIOFramework/Source/MediaIOCore/Private/MediaCapture.cpp b/Engine/Plugins/Media/MediaIOFramework/Source/MediaIOCore/Private/MediaCapture.cpp
--- a/Engine/Plugins/Media/MediaIOFramework/Source/MediaIOCore/Private/MediaCapture.cpp
+++ b/Engine/Plugins/Media/MediaIOFramework/Source/MediaIOCore/Private/MediaCapture.cpp
@@ -1372,8 +1372,6 @@ void UMediaCapture::Capture_RenderThread(FRHICommandListImmediate& RHICmdList,
 		--InMediaCapture->WaitingForRenderCommandExecutionCounter;
 	};
 
-	FTexture2DRHIRef SourceTexture;
-
 	if (CapturingFrame)
 	{
 		// Call the capture frame algo based on the specific type of resource we are using
