nDisplay: Fix media capture when not rendering in GPU0.

MediaCapture now has a notion of the valid source GPU mask of the texture/buffer. Defaults to FRHIGPUMask::All(). Can be changed any time.

diff --git a/Engine/Plugins/Media/MediaIOFramework/Source/MediaIOCore/Private/MediaCapture.cpp b/Engine/Plugins/Media/MediaIOFramework/Source/MediaIOCore/Private/MediaCapture.cpp
--- a/Engine/Plugins/Media/MediaIOFramework/Source/MediaIOCore/Private/MediaCapture.cpp
+++ b/Engine/Plugins/Media/MediaIOFramework/Source/MediaIOCore/Private/MediaCapture.cpp
@@ -662,6 +662,7 @@ FMediaCaptureOptions::FMediaCaptureOptions()
 
 UMediaCapture::UMediaCapture(const FObjectInitializer& ObjectInitializer)
 	: Super(ObjectInitializer)
+	, ValidSourceGPUMask(FRHIGPUMask::All())
 	, bOutputResourcesInitialized(false)
 	, bShouldCaptureRHIResource(false)
 	, WaitingForRenderCommandExecutionCounter(0)
@@ -1146,6 +1147,11 @@ bool UMediaCapture::HasFinishedProcessing() const
 		|| GetState() == EMediaCaptureState::Stopped;
 }
 
+void UMediaCapture::SetValidSourceGPUMask(FRHIGPUMask GPUMask)
+{
+	ValidSourceGPUMask = GPUMask;
+}
+
 void UMediaCapture::InitializeOutputResources(int32 InNumberOfBuffers)
 {
 	using namespace UE::MediaCaptureData;
@@ -1337,10 +1343,13 @@ void UMediaCapture::OnEndFrame_GameThread()
 		{
 			++WaitingForRenderCommandExecutionCounter;
 
+			const FRHIGPUMask SourceGPUMask = ValidSourceGPUMask;
+
 			// RenderCommand to be executed on the RenderThread
 			ENQUEUE_RENDER_COMMAND(FMediaOutputCaptureFrameCreateTexture)(
-				[InMediaCapture, CapturingFrame, ReadyFrame, InCapturingSceneViewport, InTextureRenderTargetResource, InDesiredSize, InOnStateChanged](FRHICommandListImmediate& RHICmdList)
+				[InMediaCapture, CapturingFrame, ReadyFrame, InCapturingSceneViewport, InTextureRenderTargetResource, InDesiredSize, InOnStateChanged, SourceGPUMask](FRHICommandListImmediate& RHICmdList)
 			{
+				SCOPED_GPU_MASK(RHICmdList, SourceGPUMask);
 				InMediaCapture->Capture_RenderThread(RHICmdList, InMediaCapture, CapturingFrame, ReadyFrame, InCapturingSceneViewport, InTextureRenderTargetResource, InDesiredSize, InOnStateChanged);
 			});
 
diff --git a/Engine/Plugins/Media/MediaIOFramework/Source/MediaIOCore/Public/MediaCapture.h b/Engine/Plugins/Media/MediaIOFramework/Source/MediaIOCore/Public/MediaCapture.h
--- a/Engine/Plugins/Media/MediaIOFramework/Source/MediaIOCore/Public/MediaCapture.h
+++ b/Engine/Plugins/Media/MediaIOFramework/Source/MediaIOCore/Public/MediaCapture.h
@@ -254,6 +254,9 @@ public:
 	UPROPERTY(BlueprintAssignable, Category = "Media|Output")
 	FMediaCaptureStateChangedSignature OnStateChanged;
 
+	/** Set the valid GPU mask of the source buffer/texture being captured. Can be called from any thread. */
+	virtual void SetValidSourceGPUMask(FRHIGPUMask GPUMask);
+
 	/**
 	 * Called when the state of the capture changed.
 	 * The callback is called on the game thread. Note that the change may occur on the rendering thread.
@@ -406,6 +409,9 @@ protected:
 	
 	/** Output size of the media capture. If a conversion is done, this resolution might be different than source resolution. i.e. 1080p RGB might be 960x1080 in YUV */
 	FIntPoint DesiredOutputSize = FIntPoint(1920, 1080);
+	
+	/** Valid source gpu mask of the source. */
+	std::atomic<FRHIGPUMask> ValidSourceGPUMask;
 
 private:
 
diff --git a/Engine/Plugins/Runtime/nDisplay/Source/DisplayClusterMedia/Private/Capture/DisplayClusterMediaCaptureBase.cpp b/Engine/Plugins/Runtime/nDisplay/Source/DisplayClusterMedia/Private/Capture/DisplayClusterMediaCaptureBase.cpp
--- a/Engine/Plugins/Runtime/nDisplay/Source/DisplayClusterMedia/Private/Capture/DisplayClusterMediaCaptureBase.cpp
+++ b/Engine/Plugins/Runtime/nDisplay/Source/DisplayClusterMedia/Private/Capture/DisplayClusterMediaCaptureBase.cpp
@@ -75,6 +75,8 @@ void FDisplayClusterMediaCaptureBase::ExportMediaData(FRHICommandListImmediate&
 
 	if (SrcTexture && DstTexture)
 	{
+		MediaCapture->SetValidSourceGPUMask(RHICmdList.GetGPUMask());
+
 		const FIntPoint SrcRegionSize = TextureInfo.Region.Size();
 
 		if (SrcTexture->GetDesc().Format == DstTexture->GetDesc().Format &&
@@ -116,6 +118,7 @@ void FDisplayClusterMediaCaptureBase::OnPostClusterTick()
 bool FDisplayClusterMediaCaptureBase::StartMediaCapture()
 {
 	FMediaCaptureOptions MediaCaptureOptions;
+
 	MediaCaptureOptions.NumberOfFramesToCapture = -1;
 	MediaCaptureOptions.bResizeSourceBuffer = true;
 	MediaCaptureOptions.bSkipFrameWhenRunningExpensiveTasks = false;
