Fix build warning

diff --git a/Engine/Plugins/Media/AjaMedia/Source/AjaMedia/Private/Player/AjaMediaPlayer.cpp b/Engine/Plugins/Media/AjaMedia/Source/AjaMedia/Private/Player/AjaMediaPlayer.cpp
--- a/Engine/Plugins/Media/AjaMedia/Source/AjaMedia/Private/Player/AjaMediaPlayer.cpp
+++ b/Engine/Plugins/Media/AjaMedia/Source/AjaMedia/Private/Player/AjaMediaPlayer.cpp
@@ -604,9 +604,9 @@ bool FAjaMediaPlayer::OnInputFrameReceived(const AJA::AJAInputFrameData& InInput
 			{
 				bool bEven = true;
 
-				if (CVarExperimentalFieldFlipFix.GetValueOnAnyThread())
+				if (FMediaIOCorePlayerBase::CVarExperimentalFieldFlipFix.GetValueOnAnyThread())
 				{
-					bEven = GFrameCounterRenderThread % 2 != CVarFlipInterlaceFields.GetValueOnAnyThread();
+					bEven = GFrameCounterRenderThread % 2 != FMediaIOCorePlayerBase::CVarFlipInterlaceFields.GetValueOnAnyThread();
 				}
 
 				if (TextureSample->InitializeInterlaced_Halfed(InVideoFrame, VideoSampleFormat, DecodedTime, VideoFrameRate, DecodedTimecode, bEven, bIsSRGBInput))
diff --git a/Engine/Plugins/Media/BlackmagicMedia/Source/BlackmagicMedia/Private/Player/BlackmagicMediaPlayer.cpp b/Engine/Plugins/Media/BlackmagicMedia/Source/BlackmagicMedia/Private/Player/BlackmagicMediaPlayer.cpp
--- a/Engine/Plugins/Media/BlackmagicMedia/Source/BlackmagicMedia/Private/Player/BlackmagicMediaPlayer.cpp
+++ b/Engine/Plugins/Media/BlackmagicMedia/Source/BlackmagicMedia/Private/Player/BlackmagicMediaPlayer.cpp
@@ -357,9 +357,9 @@ namespace BlackmagicMediaPlayerHelpers
 					{
 						bool bEven = true;
 
-						if (CVarExperimentalFieldFlipFix.GetValueOnAnyThread())
+						if (FMediaIOCorePlayerBase::CVarExperimentalFieldFlipFix.GetValueOnAnyThread())
 						{
-							bEven = GFrameCounterRenderThread % 2 != CVarFlipInterlaceFields.GetValueOnAnyThread();
+							bEven = GFrameCounterRenderThread % 2 != FMediaIOCorePlayerBase::CVarFlipInterlaceFields.GetValueOnAnyThread();
 						}
 
 						auto TextureSampleEven = MediaPlayer->TextureSamplePool->AcquireShared();
diff --git a/Engine/Plugins/Media/MediaIOFramework/Source/MediaIOCore/Private/Player/MediaIOCorePlayerBase.cpp b/Engine/Plugins/Media/MediaIOFramework/Source/MediaIOCore/Private/Player/MediaIOCorePlayerBase.cpp
--- a/Engine/Plugins/Media/MediaIOFramework/Source/MediaIOCore/Private/Player/MediaIOCorePlayerBase.cpp
+++ b/Engine/Plugins/Media/MediaIOFramework/Source/MediaIOCore/Private/Player/MediaIOCorePlayerBase.cpp
@@ -74,6 +74,16 @@ const FName FMediaIOCoreMediaOption::ResolutionWidth("ResolutionWidth");
 const FName FMediaIOCoreMediaOption::ResolutionHeight("ResolutionHeight");
 const FName FMediaIOCoreMediaOption::VideoModeName("VideoModeName");
 
+TAutoConsoleVariable<int32> FMediaIOCorePlayerBase::CVarFlipInterlaceFields(
+	TEXT("MediaIO.FlipInputInterlaceFields"), 0,
+	TEXT("Whether to flip input interlace fields. (Experimental)"),
+	ECVF_RenderThreadSafe);
+
+TAutoConsoleVariable<int32> FMediaIOCorePlayerBase::CVarExperimentalFieldFlipFix(
+	TEXT("MediaIO.PreventFieldFlipping"), 1,
+	TEXT("Whether to attempt fixing field flipping on input. (Experimental)"),
+	ECVF_RenderThreadSafe);
+
 /* FMediaIOCorePlayerBase structors
  *****************************************************************************/
 
diff --git a/Engine/Plugins/Media/MediaIOFramework/Source/MediaIOCore/Public/MediaIOCorePlayerBase.h b/Engine/Plugins/Media/MediaIOFramework/Source/MediaIOCore/Public/MediaIOCorePlayerBase.h
--- a/Engine/Plugins/Media/MediaIOFramework/Source/MediaIOCore/Public/MediaIOCorePlayerBase.h
+++ b/Engine/Plugins/Media/MediaIOFramework/Source/MediaIOCore/Public/MediaIOCorePlayerBase.h
@@ -18,16 +18,6 @@
 
 #include "Misc/FrameRate.h"
 
-static TAutoConsoleVariable<int32> CVarFlipInterlaceFields(
-	TEXT("MediaIO.FlipInputInterlaceFields"), 0,
-	TEXT("Whether to flip input interlace fields. (Experimental)"),
-	ECVF_RenderThreadSafe);
-
-static TAutoConsoleVariable<int32> CVarExperimentalFieldFlipFix(
-	TEXT("MediaIO.PreventFieldFlipping"), 1,
-	TEXT("Whether to attempt fixing field flipping on input. (Experimental)"),
-	ECVF_RenderThreadSafe);
-
 class FMediaIOCoreSamples;
 class IMediaEventSink;
 
@@ -262,6 +252,9 @@ protected:
 	/** Pool of textures registerd with GPU Texture transfer. */
 	TArray<TRefCountPtr<FRHITexture>> Textures;
 
+	static TAutoConsoleVariable<int32> CVarFlipInterlaceFields;
+	static TAutoConsoleVariable<int32> CVarExperimentalFieldFlipFix;
+
 private:
 	/** GPU Texture transfer object */
 	UE::GPUTextureTransfer::TextureTransferPtr GPUTextureTransfer;
@@ -269,4 +262,5 @@ private:
 	TSet<void*> RegisteredBuffers;
 	/** Pool of textures registerd with GPU Texture transfer. */
 	TSet<TRefCountPtr<FRHITexture>> RegisteredTextures;
+
 };
