Fix CIS by moving static declaration from protected to public

diff --git a/Engine/Plugins/Media/MediaIOFramework/Source/MediaIOCore/Public/MediaIOCorePlayerBase.h b/Engine/Plugins/Media/MediaIOFramework/Source/MediaIOCore/Public/MediaIOCorePlayerBase.h
--- a/Engine/Plugins/Media/MediaIOFramework/Source/MediaIOCore/Public/MediaIOCorePlayerBase.h
+++ b/Engine/Plugins/Media/MediaIOFramework/Source/MediaIOCore/Public/MediaIOCorePlayerBase.h
@@ -63,6 +63,9 @@ class MEDIAIOCORE_API FMediaIOCorePlayerBase
 	, public ITimedDataInput
 	, public TSharedFromThis<FMediaIOCorePlayerBase>
 {
+public:
+	static TAutoConsoleVariable<int32> CVarFlipInterlaceFields;
+	static TAutoConsoleVariable<int32> CVarExperimentalFieldFlipFix;
 public:
 
 	/**
@@ -252,9 +255,6 @@ protected:
 	/** Pool of textures registerd with GPU Texture transfer. */
 	TArray<TRefCountPtr<FRHITexture>> Textures;
 
-	static TAutoConsoleVariable<int32> CVarFlipInterlaceFields;
-	static TAutoConsoleVariable<int32> CVarExperimentalFieldFlipFix;
-
 private:
 	/** GPU Texture transfer object */
 	UE::GPUTextureTransfer::TextureTransferPtr GPUTextureTransfer;
