Delay RHI media capture by 1 frame in order to fix tearing in media capture using gpu direct

diff --git a/Engine/Plugins/Media/AjaMedia/Source/AjaMediaOutput/Private/AjaMediaCapture.cpp b/Engine/Plugins/Media/AjaMedia/Source/AjaMediaOutput/Private/AjaMediaCapture.cpp
--- a/Engine/Plugins/Media/AjaMedia/Source/AjaMediaOutput/Private/AjaMediaCapture.cpp
+++ b/Engine/Plugins/Media/AjaMedia/Source/AjaMediaOutput/Private/AjaMediaCapture.cpp
@@ -325,25 +325,6 @@ bool UAjaMediaCapture::ShouldCaptureRHIResource() const
 	return bGPUTextureTransferAvailable && CVarAjaEnableGPUDirect.GetValueOnAnyThread() == 1;
 }
 
-void UAjaMediaCapture::BeforeFrameCaptured_RenderingThread(const FCaptureBaseData& InBaseData, TSharedPtr<FMediaCaptureUserData, ESPMode::ThreadSafe> InUserData, FTextureRHIRef InTexture)
-{
-	if (ShouldCaptureRHIResource())
-	{
-		if (!TexturesToRelease.Contains(InTexture))
-		{
-			TexturesToRelease.Add(InTexture);
-
-			FRHITexture2D* Texture = InTexture->GetTexture2D();
-			AJA::FRegisterDMATextureArgs Args;
-			Args.RHITexture = Texture->GetNativeResource();
-			//Args.InRHIResourceMemory = Texture->GetNativeResource(); todo: VulkanTexture->Surface->GetAllocationHandle for Vulkan
-			AJA::RegisterDMATexture(Args);
-		}
-	}
-
-	AJA::LockDMATexture(InTexture->GetTexture2D()->GetNativeResource());
-}
-
 void UAjaMediaCapture::ApplyViewportTextureAlpha(TSharedPtr<FSceneViewport> InSceneViewport)
 {
 	if (InSceneViewport.IsValid())
@@ -547,7 +528,7 @@ void UAjaMediaCapture::OnFrameCaptured_RenderingThread(const FCaptureBaseData& I
 	}
 }
 
-void UAjaMediaCapture::OnRHIResourceCaptured_RenderingThread(const FCaptureBaseData& InBaseData,	TSharedPtr<FMediaCaptureUserData, ESPMode::ThreadSafe> InUserData, FTextureRHIRef InTexture)
+void UAjaMediaCapture::OnRHIResourceCaptured_RenderingThread(const FCaptureBaseData& InBaseData, TSharedPtr<FMediaCaptureUserData, ESPMode::ThreadSafe> InUserData, FTextureRHIRef InTexture)
 {
 	TRACE_CPUPROFILER_EVENT_SCOPE(UAjaMediaCapture::OnFrameCaptured_RenderingThread);
 	
@@ -555,11 +536,6 @@ void UAjaMediaCapture::OnRHIResourceCaptured_RenderingThread(const FCaptureBaseD
 	FScopeLock ScopeLock(&RenderThreadCriticalSection);
 	if (OutputChannel)
 	{
-		{
-			TRACE_CPUPROFILER_EVENT_SCOPE(UAjaMediaCapture::OnFrameCaptured_RenderingThread::UnlockDMATexture);
-			AJA::UnlockDMATexture(InTexture->GetTexture2D()->GetNativeResource());
-		}
-
 		const AJA::FTimecode Timecode = AjaMediaCaptureDevice::ConvertToAJATimecode(InBaseData.SourceFrameTimecode, InBaseData.SourceFrameTimecodeFramerate.AsDecimal(), FrameRate.AsDecimal());
 		
 		AJA::AJAOutputFrameBufferData FrameBuffer;
@@ -588,6 +564,31 @@ void UAjaMediaCapture::OnRHIResourceCaptured_RenderingThread(const FCaptureBaseD
 	}
 }
 
+void UAjaMediaCapture::LockDMATexture_RenderThread(FTextureRHIRef InTexture)
+{
+	if (ShouldCaptureRHIResource())
+	{
+		if (!TexturesToRelease.Contains(InTexture))
+		{
+			TexturesToRelease.Add(InTexture);
+
+			FRHITexture2D* Texture = InTexture->GetTexture2D();
+			AJA::FRegisterDMATextureArgs Args;
+			Args.RHITexture = Texture->GetNativeResource();
+			//Args.InRHIResourceMemory = Texture->GetNativeResource(); todo: VulkanTexture->Surface->GetAllocationHandle for Vulkan
+			AJA::RegisterDMATexture(Args);
+		}
+	}
+
+	AJA::LockDMATexture(InTexture->GetTexture2D()->GetNativeResource());
+}
+
+void UAjaMediaCapture::UnlockDMATexture_RenderThread(FTextureRHIRef InTexture)
+{
+	TRACE_CPUPROFILER_EVENT_SCOPE(UAjaMediaCapture::OnFrameCaptured_RenderingThread::UnlockDMATexture);
+	AJA::UnlockDMATexture(InTexture->GetTexture2D()->GetNativeResource());
+}
+
 void UAjaMediaCapture::WaitForSync_RenderingThread() const
 {
 	if (bWaitForSyncEvent)
diff --git a/Engine/Plugins/Media/AjaMedia/Source/AjaMediaOutput/Public/AjaMediaCapture.h b/Engine/Plugins/Media/AjaMedia/Source/AjaMediaOutput/Public/AjaMediaCapture.h
--- a/Engine/Plugins/Media/AjaMedia/Source/AjaMediaOutput/Public/AjaMediaCapture.h
+++ b/Engine/Plugins/Media/AjaMedia/Source/AjaMediaOutput/Public/AjaMediaCapture.h
@@ -43,9 +43,10 @@ protected:
 	virtual void StopCaptureImpl(bool bAllowPendingFrameToBeProcess) override;
 	virtual bool ShouldCaptureRHIResource() const override;
 	
-	virtual void BeforeFrameCaptured_RenderingThread(const FCaptureBaseData& InBaseData, TSharedPtr<FMediaCaptureUserData, ESPMode::ThreadSafe> InUserData, FTextureRHIRef InTexture) override;
 	virtual void OnFrameCaptured_RenderingThread(const FCaptureBaseData& InBaseData, TSharedPtr<FMediaCaptureUserData, ESPMode::ThreadSafe> InUserData, void* InBuffer, int32 Width, int32 Height, int32 BytesPerRow) override;
 	virtual void OnRHIResourceCaptured_RenderingThread(const FCaptureBaseData& InBaseData, TSharedPtr<FMediaCaptureUserData, ESPMode::ThreadSafe> InUserData, FTextureRHIRef InTexture) override;
+	virtual void LockDMATexture_RenderThread(FTextureRHIRef InTexture) override;
+	virtual void UnlockDMATexture_RenderThread(FTextureRHIRef InTexture) override;
 
 private:
 	struct FAjaOutputCallback;
diff --git a/Engine/Plugins/Media/BlackmagicMedia/Source/BlackmagicMediaOutput/Private/BlackmagicMediaCapture.cpp b/Engine/Plugins/Media/BlackmagicMedia/Source/BlackmagicMediaOutput/Private/BlackmagicMediaCapture.cpp
--- a/Engine/Plugins/Media/BlackmagicMedia/Source/BlackmagicMediaOutput/Private/BlackmagicMediaCapture.cpp
+++ b/Engine/Plugins/Media/BlackmagicMedia/Source/BlackmagicMediaOutput/Private/BlackmagicMediaCapture.cpp
@@ -581,7 +581,7 @@ bool UBlackmagicMediaCapture::InitBlackmagic(UBlackmagicMediaOutput* InBlackmagi
 	return true;
 }
 
-void UBlackmagicMediaCapture::BeforeFrameCaptured_RenderingThread(const FCaptureBaseData& InBaseData, TSharedPtr<FMediaCaptureUserData, ESPMode::ThreadSafe> InUserData, FTextureRHIRef InTexture)
+void UBlackmagicMediaCapture::LockDMATexture_RenderThread(FTextureRHIRef InTexture)
 {
 	if (ShouldCaptureRHIResource())
 	{
@@ -604,6 +604,12 @@ void UBlackmagicMediaCapture::BeforeFrameCaptured_RenderingThread(const FCapture
 	}
 }
 
+void UBlackmagicMediaCapture::UnlockDMATexture_RenderThread(FTextureRHIRef InTexture)
+{
+	TRACE_CPUPROFILER_EVENT_SCOPE(UBlackmagicMediaCapture::UnlockDMATexture);
+	BlackmagicDesign::UnlockDMATexture(InTexture->GetTexture2D()->GetNativeResource());
+}
+
 void UBlackmagicMediaCapture::OnFrameCaptured_RenderingThread(const FCaptureBaseData& InBaseData, TSharedPtr<FMediaCaptureUserData, ESPMode::ThreadSafe> InUserData, void* InBuffer, int32 Width, int32 Height, int32 BytesPerRow)
 {
 	// Prevent the rendering thread from copying while we are stopping the capture.
@@ -688,11 +694,6 @@ void UBlackmagicMediaCapture::OnRHIResourceCaptured_RenderingThread(const FCaptu
 	// Prevent the rendering thread from copying while we are stopping the capture.
 	TRACE_CPUPROFILER_EVENT_SCOPE(UBlackmagicMediaCapture::OnFrameCaptured_RenderingThread);
 
-	{
-		TRACE_CPUPROFILER_EVENT_SCOPE(UBlackmagicMediaCapture::UnlockDMATexture);
-		BlackmagicDesign::UnlockDMATexture(InTexture->GetTexture2D()->GetNativeResource());
-	}
-
 	FScopeLock ScopeLock(&RenderThreadCriticalSection);
 
 
diff --git a/Engine/Plugins/Media/BlackmagicMedia/Source/BlackmagicMediaOutput/Public/BlackmagicMediaCapture.h b/Engine/Plugins/Media/BlackmagicMedia/Source/BlackmagicMediaOutput/Public/BlackmagicMediaCapture.h
--- a/Engine/Plugins/Media/BlackmagicMedia/Source/BlackmagicMediaOutput/Public/BlackmagicMediaCapture.h
+++ b/Engine/Plugins/Media/BlackmagicMedia/Source/BlackmagicMediaOutput/Public/BlackmagicMediaCapture.h
@@ -43,9 +43,10 @@ protected:
 	virtual void StopCaptureImpl(bool bAllowPendingFrameToBeProcess) override;
 	virtual bool ShouldCaptureRHIResource() const override;
 
-	virtual void BeforeFrameCaptured_RenderingThread(const FCaptureBaseData& InBaseData, TSharedPtr<FMediaCaptureUserData, ESPMode::ThreadSafe> InUserData, FTextureRHIRef InTexture) override;
 	virtual void OnFrameCaptured_RenderingThread(const FCaptureBaseData& InBaseData, TSharedPtr<FMediaCaptureUserData, ESPMode::ThreadSafe> InUserData, void* InBuffer, int32 Width, int32 Height, int32 BytesPerRow) override;
 	virtual void OnRHIResourceCaptured_RenderingThread(const FCaptureBaseData& InBaseData, TSharedPtr<FMediaCaptureUserData, ESPMode::ThreadSafe> InUserData, FTextureRHIRef InTexture) override;
+	virtual void LockDMATexture_RenderThread(FTextureRHIRef InTexture) override;
+	virtual void UnlockDMATexture_RenderThread(FTextureRHIRef InTexture) override;
 
 private:
 	struct FBlackmagicOutputCallback;
diff --git a/Engine/Plugins/Media/MediaIOFramework/Source/MediaIOCore/Private/MediaCapture.cpp b/Engine/Plugins/Media/MediaIOFramework/Source/MediaIOCore/Private/MediaCapture.cpp
--- a/Engine/Plugins/Media/MediaIOFramework/Source/MediaIOCore/Private/MediaCapture.cpp
+++ b/Engine/Plugins/Media/MediaIOFramework/Source/MediaIOCore/Private/MediaCapture.cpp
@@ -136,6 +136,7 @@ namespace UE::MediaCaptureData
 	public:
 		FCaptureFrame()
 			: bReadbackRequested(false)
+			, bDoingGPUCopy(false)
 		{
 
 		}
@@ -157,8 +158,12 @@ namespace UE::MediaCaptureData
 		/** Returns true if the readback is ready to be used */
 		virtual bool IsReadbackReady(FRHIGPUMask GPUMask) = 0;
 
+		virtual FRHITexture* GetTextureResource() { return nullptr; }
+		virtual FRHIBuffer* GetBufferResource() { return nullptr; }
+
 		UMediaCapture::FCaptureBaseData CaptureBaseData;
 		TAtomic<bool> bReadbackRequested;
+		TAtomic<bool> bDoingGPUCopy;
 		TSharedPtr<FMediaCaptureUserData> UserData;
 	};
 
@@ -221,6 +226,11 @@ namespace UE::MediaCaptureData
 			AddEnqueueCopyPass(RDGBuilder, ReadbackTexture.Get(), ResourceToReadback);
 		}
 
+		virtual FRHITexture* GetTextureResource() override
+		{
+			return GetRHIResource();
+		}
+
 		/** Returns RHI resource of the allocated pooled resource */
 		FRHITexture* GetRHIResource()
 		{
@@ -289,6 +299,11 @@ namespace UE::MediaCaptureData
 			AddEnqueueCopyPass(RDGBuilder, ReadbackBuffer.Get(), ResourceToReadback, Buffer->GetRHI()->GetSize());
 		}
 
+		virtual FRHIBuffer* GetBufferResource() override
+		{
+			return GetRHIResource();
+		}
+
 		/** Returns RHI resource of the allocated pooled resource */
 		FRHIBuffer* GetRHIResource()
 		{
@@ -326,19 +341,6 @@ namespace UE::MediaCaptureData
 	{
 	public:
 
-		static void BeforeFrameCapture(const FCaptureFrameArgs& Args, FTextureCaptureFrame* CapturingFrame)
-		{
-			UE_LOG(LogMediaIOCore, VeryVerbose, TEXT("Before Locking texture %p"), reinterpret_cast<uintptr_t>(CapturingFrame->RenderTarget->GetRHI()->GetNativeResource()));
-			Args.MediaCapture->BeforeFrameCaptured_RenderingThread(CapturingFrame->CaptureBaseData, CapturingFrame->UserData, CapturingFrame->RenderTarget->GetRHI());
-			UE_LOG(LogMediaIOCore, VeryVerbose, TEXT("After Locking texture %p"), reinterpret_cast<uintptr_t>(CapturingFrame->RenderTarget->GetRHI()->GetNativeResource()));
-		}
-
-		static void BeforeFrameCapture(const FCaptureFrameArgs& Args, FBufferCaptureFrame* CapturingFrame)
-		{
-			FBufferCaptureFrame* BufferFrame = static_cast<FBufferCaptureFrame*>(CapturingFrame);
-			Args.MediaCapture->BeforeFrameCaptured_RenderingThread(CapturingFrame->CaptureBaseData, CapturingFrame->UserData, BufferFrame->Buffer->GetRHI());
-		}
-
 		static FTexture2DRHIRef GetSourceTextureForInput(const FCaptureFrameArgs& Args)
 		{
 			FTexture2DRHIRef SourceTexture;
@@ -564,8 +566,12 @@ namespace UE::MediaCaptureData
 			}
 
 			{
-				TRACE_CPUPROFILER_EVENT_SCOPE(UMediaCapture::BeforeFrameCaptured);
-				BeforeFrameCapture(Args, CapturingFrame);
+				TRACE_CPUPROFILER_EVENT_SCOPE(UMediaCapture::LockDMATexture_RenderThread);
+
+				if (CapturingFrame->IsTextureResource())
+				{
+					Args.MediaCapture->LockDMATexture_RenderThread(CapturingFrame->GetTextureResource());
+				}
 			}
 
 			{
@@ -616,15 +622,9 @@ namespace UE::MediaCaptureData
 					++Args.MediaCapture->NumberOfResourcesToReadback;
 				}
 
-				GraphBuilder.Execute();
+				CapturingFrame->bDoingGPUCopy = true;
 
-				if (Args.MediaCapture->bShouldCaptureRHIResource)
-				{
-					TRACE_CPUPROFILER_EVENT_SCOPE(UMediaCapture::RHIResourceCaptured);
-					SCOPE_CYCLE_COUNTER(STAT_MediaCapture_RenderThread_CaptureCallback)
-					Args.MediaCapture->OnRHIResourceCaptured_RenderingThread(CapturingFrame->CaptureBaseData, CapturingFrame->UserData, CapturingFrame->GetRHIResource());
-					CapturingFrame->bReadbackRequested = false;
-				}
+				GraphBuilder.Execute();
 			}
 
 			return true;
@@ -1417,7 +1417,7 @@ void UMediaCapture::Capture_RenderThread(FRHICommandListImmediate& RHICmdList,
 		}
 	}
 
-	if (!InMediaCapture->bShouldCaptureRHIResource && InMediaCapture->GetState() != EMediaCaptureState::Error)
+	if (InMediaCapture->GetState() != EMediaCaptureState::Error)
 	{
 		if (ReadyFrame->bReadbackRequested)
 		{
@@ -1459,6 +1459,23 @@ void UMediaCapture::Capture_RenderThread(FRHICommandListImmediate& RHICmdList,
 
 			ReadyFrame->Unlock();
 		}
+		else if (InMediaCapture->bShouldCaptureRHIResource && ReadyFrame->bDoingGPUCopy)
+		{
+			if (ReadyFrame->IsTextureResource())
+			{
+				FTextureCaptureFrame* TextureFrame = static_cast<FTextureCaptureFrame*>(ReadyFrame);
+				{
+					TRACE_CPUPROFILER_EVENT_SCOPE(UMediaCapture::UnlockDMATexture_RenderThread);
+					UnlockDMATexture_RenderThread(TextureFrame->RenderTarget->GetRHI());
+				}
+
+				TRACE_CPUPROFILER_EVENT_SCOPE(UMediaCapture::RHIResourceCaptured);
+				SCOPE_CYCLE_COUNTER(STAT_MediaCapture_RenderThread_CaptureCallback)
+					InMediaCapture->OnRHIResourceCaptured_RenderingThread(ReadyFrame->CaptureBaseData, ReadyFrame->UserData, TextureFrame->GetRHIResource());
+
+				ReadyFrame->bDoingGPUCopy = false;
+			}
+		}
 	}
 }
 
diff --git a/Engine/Plugins/Media/MediaIOFramework/Source/MediaIOCore/Public/MediaCapture.h b/Engine/Plugins/Media/MediaIOFramework/Source/MediaIOCore/Public/MediaCapture.h
--- a/Engine/Plugins/Media/MediaIOFramework/Source/MediaIOCore/Public/MediaCapture.h
+++ b/Engine/Plugins/Media/MediaIOFramework/Source/MediaIOCore/Public/MediaCapture.h
@@ -277,6 +277,10 @@ protected:
 	virtual bool UpdateRenderTargetImpl(UTextureRenderTarget2D* InRenderTarget) { return true; }
 	virtual void StopCaptureImpl(bool bAllowPendingFrameToBeProcess) { }
 
+	//~ DMA Functions implemented in child classes, Temporary until AJA and BM middleman are migrated to a module.
+	virtual void LockDMATexture_RenderThread(FTextureRHIRef InTexture) {}
+	virtual void UnlockDMATexture_RenderThread(FTextureRHIRef InTexture) {}
+
 	friend class UE::MediaCaptureData::FCaptureFrame;
 	friend class UE::MediaCaptureData::FMediaCaptureHelper;
 	struct FCaptureBaseData
