Expose option to enable linear conversion

diff --git a/Engine/Plugins/Media/MediaIOFramework/Source/MediaIOCore/Private/MediaCapture.cpp b/Engine/Plugins/Media/MediaIOFramework/Source/MediaIOCore/Private/MediaCapture.cpp
--- a/Engine/Plugins/Media/MediaIOFramework/Source/MediaIOCore/Private/MediaCapture.cpp
+++ b/Engine/Plugins/Media/MediaIOFramework/Source/MediaIOCore/Private/MediaCapture.cpp
@@ -477,7 +477,7 @@ namespace UE::MediaCaptureData
 			else
 			{
 				//At some point we should support color conversion (ocio) but for now we push incoming texture as is
-				constexpr bool bDoLinearToSRGB = false;
+				const bool bDoLinearToSRGB = Args.MediaCapture->DesiredCaptureOptions.bApplyLinearToSRGBConversion;
 
 				FGlobalShaderMap* GlobalShaderMap = GetGlobalShaderMap(GMaxRHIFeatureLevel);
 				TShaderMapRef<FScreenPassVS> VertexShader(GlobalShaderMap);
@@ -605,7 +605,7 @@ namespace UE::MediaCaptureData
 					else
 					{
 						TRACE_CPUPROFILER_EVENT_SCOPE(UMediaCapture::FormatConversion);
-						AddConversionPass(Args, { GraphBuilder, SourceRGBTexture, bRequiresFormatConversion, CopyInfo, SizeU, SizeV }, OutputResource);
+						AddConversionPass(Args, { GraphBuilder, SourceRGBTexture, bRequiresFormatConversion,  CopyInfo, SizeU, SizeV }, OutputResource);
 					}
 				}
 
diff --git a/Engine/Plugins/Media/MediaIOFramework/Source/MediaIOCore/Public/MediaCapture.h b/Engine/Plugins/Media/MediaIOFramework/Source/MediaIOCore/Public/MediaCapture.h
--- a/Engine/Plugins/Media/MediaIOFramework/Source/MediaIOCore/Public/MediaCapture.h
+++ b/Engine/Plugins/Media/MediaIOFramework/Source/MediaIOCore/Public/MediaCapture.h
@@ -142,6 +142,12 @@ public:
 	UPROPERTY(BlueprintReadWrite, EditAnywhere, Category="MediaCapture")
 	bool bForceAlphaToOneOnConversion;
 
+	/**
+	 * Whether to apply a linear to sRGB conversion to the texture before outputting. 
+	 */
+	UPROPERTY(BlueprintReadWrite, EditAnywhere, Category = "MediaCapture")
+	bool bApplyLinearToSRGBConversion = false;;
+
 	/** Automatically stop capturing after a predetermined number of images. */
 	UPROPERTY(BlueprintReadWrite, EditAnywhere, Category = "MediaCapture")
 	bool bAutostopOnCapture;
