Add display name to media capture variable

diff --git a/Engine/Plugins/Media/MediaIOFramework/Source/MediaIOCore/Public/MediaCapture.h b/Engine/Plugins/Media/MediaIOFramework/Source/MediaIOCore/Public/MediaCapture.h
--- a/Engine/Plugins/Media/MediaIOFramework/Source/MediaIOCore/Public/MediaCapture.h
+++ b/Engine/Plugins/Media/MediaIOFramework/Source/MediaIOCore/Public/MediaCapture.h
@@ -145,8 +145,8 @@ public:
 	/**
 	 * Whether to apply a linear to sRGB conversion to the texture before outputting. 
 	 */
-	UPROPERTY(BlueprintReadWrite, EditAnywhere, Category = "MediaCapture")
-	bool bApplyLinearToSRGBConversion = false;;
+	UPROPERTY(BlueprintReadWrite, EditAnywhere, Category = "MediaCapture", meta=(DisplayName="Apply linear to sRGB conversion"))
+	bool bApplyLinearToSRGBConversion = false;
 
 	/** Automatically stop capturing after a predetermined number of images. */
 	UPROPERTY(BlueprintReadWrite, EditAnywhere, Category = "MediaCapture")
