Fix audio crash when outputting audio with blackmagic on linux

diff --git a/Engine/Plugins/Media/MediaIOFramework/Source/MediaIOCore/Private/MediaIOCoreAudioOutput.cpp b/Engine/Plugins/Media/MediaIOFramework/Source/MediaIOCore/Private/MediaIOCoreAudioOutput.cpp
--- a/Engine/Plugins/Media/MediaIOFramework/Source/MediaIOCore/Private/MediaIOCoreAudioOutput.cpp
+++ b/Engine/Plugins/Media/MediaIOFramework/Source/MediaIOCore/Private/MediaIOCoreAudioOutput.cpp
@@ -7,13 +7,16 @@
 #include "AudioMixerSubmix.h"
 #include "Sound/AudioSettings.h"
 
-
 #if WITH_EDITOR
 #include "Editor.h"
+#include "Framework/Notifications/NotificationManager.h"
+#include "Widgets/Notifications/SNotificationList.h"
 #endif
 
 DEFINE_LOG_CATEGORY_STATIC(LogMediaIOAudioOutput, Log, All);
 
+#define LOCTEXT_NAMESPACE "MediaIOAudioOutput"
+
 FMediaIOAudioOutput::FMediaIOAudioOutput(Audio::FPatchOutputStrongPtr InPatchOutput, const FAudioOptions& InAudioOptions)
 	: PatchOutput(MoveTemp(InPatchOutput))
     , NumInputChannels(InAudioOptions.InNumInputChannels)
@@ -87,12 +90,27 @@ void FMediaIOAudioCapture::OnNewSubmixBuffer(const USoundSubmix* InOwningSubmix,
 
 TSharedPtr<FMediaIOAudioOutput> FMediaIOAudioCapture::CreateAudioOutput(int32 InNumOutputChannels, FFrameRate InTargetFrameRate, uint32 InMaxSampleLatency, uint32 InOutputSampleRate)
 {
+	if (NumChannels > InNumOutputChannels)
+	{
+		UE_LOG(LogMediaIOAudioOutput, Error, TEXT("Audio capture initialization error, please change the audio output channel count to a number greater or equal to %d."), NumChannels);
+
+#if WITH_EDITOR
+		const FText WarningText = FText::Format(LOCTEXT("AudioCaptureError", "Audio capture initialization error, please change the audio output channel count to a number greather or equal to %d."), NumChannels);
+		FNotificationInfo WarningNotification(WarningText);
+		WarningNotification.bFireAndForget = true;
+		WarningNotification.ExpireDuration = 6.0f;
+		WarningNotification.bUseThrobber = false;
+
+		// For adding notifications.
+		FSlateNotificationManager::Get().AddNotification(WarningNotification);
+#endif
+		return nullptr;
+	}
+	
 	if (ensureMsgf(InOutputSampleRate == SampleRate, TEXT("The engine's sample rate is different from the output sample rate and resampling is not yet supported in Media Captutre.")))
 	{
 		constexpr float Gain = 1.0f;
-		
-		checkf(NumChannels <= InNumOutputChannels, TEXT("At the moment MediaIOAudioCapture only supports up mixing."));
-		check(InNumOutputChannels != 0);
+		check(InNumOutputChannels > 0);
 		
 		Audio::FPatchOutputStrongPtr PatchOutput = AudioSplitter.AddNewPatch(InMaxSampleLatency, Gain);
 		FMediaIOAudioOutput::FAudioOptions Options;
@@ -167,3 +185,6 @@ void FMediaIOAudioCapture::UnregisterBufferListener(FAudioDevice* AudioDevice)
 		AudioDevice->UnregisterSubmixBufferListener(this);
 	}
 }
+
+
+#undef LOCTEXT_NAMESPACE //MediaIOAudioOutput
\ No newline at end of file
diff --git a/Engine/Plugins/Media/MediaIOFramework/Source/MediaIOCore/Public/MediaIOCoreAudioOutput.h b/Engine/Plugins/Media/MediaIOFramework/Source/MediaIOCore/Public/MediaIOCoreAudioOutput.h
--- a/Engine/Plugins/Media/MediaIOFramework/Source/MediaIOCore/Public/MediaIOCoreAudioOutput.h
+++ b/Engine/Plugins/Media/MediaIOFramework/Source/MediaIOCore/Public/MediaIOCoreAudioOutput.h
@@ -21,7 +21,8 @@ namespace UE::MediaIoCoreModule::Private
 		// @Note: We only support upmixing at the moment.
 		
 		TArray<OutputType> ConvertedBuffer;
-		ConvertedBuffer.SetNumZeroed(NumOutputChannels / NumInputChannels * InBuffer.Num());
+		const float ChannelRatio = static_cast<float>(NumOutputChannels) / NumInputChannels;
+		ConvertedBuffer.SetNumZeroed(FMath::CeilToInt32(ChannelRatio * InBuffer.Num()));
 
 		OutputType* ConvertedBufferPtr = ConvertedBuffer.GetData();
 
