- Adding immediate rhi resource capture mode to MediaCapture

diff --git a/Engine/Plugins/Media/MediaIOFramework/Source/MediaIOCore/Private/MediaCapture.cpp b/Engine/Plugins/Media/MediaIOFramework/Source/MediaIOCore/Private/MediaCapture.cpp
--- a/Engine/Plugins/Media/MediaIOFramework/Source/MediaIOCore/Private/MediaCapture.cpp
+++ b/Engine/Plugins/Media/MediaIOFramework/Source/MediaIOCore/Private/MediaCapture.cpp
@@ -134,8 +134,9 @@ namespace UE::MediaCaptureData
 	class FCaptureFrame
 	{
 	public:
-		FCaptureFrame()
-			: bReadbackRequested(false)
+		FCaptureFrame(int32 InFrameId)
+			: FrameId(InFrameId)
+			, bReadbackRequested(false)
 			, bDoingGPUCopy(false)
 		{
 
@@ -161,6 +162,9 @@ namespace UE::MediaCaptureData
 		virtual FRHITexture* GetTextureResource() { return nullptr; }
 		virtual FRHIBuffer* GetBufferResource() { return nullptr; }
 
+		bool IsPending() const { return bReadbackRequested || bDoingGPUCopy; }
+
+		int32 FrameId = 0;
 		UMediaCapture::FCaptureBaseData CaptureBaseData;
 		TAtomic<bool> bReadbackRequested;
 		TAtomic<bool> bDoingGPUCopy;
@@ -173,6 +177,11 @@ namespace UE::MediaCaptureData
 		/** Type alias for the output resource type used during capture frame */
 		using FOutputResourceType = FRDGTextureRef;
 
+		FTextureCaptureFrame(int32 InFrameId)
+			: FCaptureFrame(InFrameId)
+		{
+		}
+
 		//~ Begin FCaptureFrame interface
 		virtual bool HasValidResource() const override
 		{
@@ -248,6 +257,11 @@ namespace UE::MediaCaptureData
 		/** Type alias for the output resource type used during capture frame */
 		using FOutputResourceType = FRDGBufferRef;
 
+		FBufferCaptureFrame(int32 InFrameId)
+			: FCaptureFrame(InFrameId)
+		{
+		}
+
 		//~ Begin FCaptureFrame interface
 		virtual bool HasValidResource() const override
 		{
@@ -268,7 +282,7 @@ namespace UE::MediaCaptureData
 		{
 			if (ReadbackBuffer->IsReady(GPUMask) == false)
 			{
-				UE_LOG(LogMediaIOCore, Verbose, TEXT("Fenced for buffer readback was not ready, blocking."));
+				UE_LOG(LogMediaIOCore, Verbose, TEXT("Fence for buffer readback was not ready, blocking."));
 				RHICmdList.BlockUntilGPUIdle();
 			}
 
@@ -314,21 +328,87 @@ namespace UE::MediaCaptureData
 		TRefCountPtr<FRDGPooledBuffer> Buffer;
 		TUniquePtr<FRHIGPUBufferReadback> ReadbackBuffer;
 	};
+	
+	class FFrameManager
+	{
+	public:
+
+		TConstArrayView<TUniquePtr<FCaptureFrame>> GetFrames() const
+		{
+			return CaptureFrames;
+		}
+
+		void AddFrame(TUniquePtr<FCaptureFrame>&& NewFrame)
+		{
+			AvailableFrames.Enqueue(NewFrame->FrameId);
+			CaptureFrames.Add(MoveTemp(NewFrame));
+		}
+
+		void MarkAvailable(const FCaptureFrame& InFrame)
+		{
+			if (ensure(!InFrame.IsPending()))
+			{
+				AvailableFrames.Enqueue(InFrame.FrameId);
+			}
+		}
+
+		void MarkPending(const FCaptureFrame& InFrame)
+		{
+			if (ensure(InFrame.IsPending()))
+			{
+				PendingFrames.Enqueue(InFrame.FrameId);
+			}
+		}
+
+		void CompleteNextPending(const FCaptureFrame& InFrame)
+		{
+			TOptional<int32> NextPending = PendingFrames.Dequeue();
+			if (ensure(NextPending.IsSet() && InFrame.FrameId == NextPending.GetValue()))
+			{
+				MarkAvailable(InFrame);
+			}
+		}
+
+		FCaptureFrame* GetNextAvailable()
+		{
+			TOptional<int32> NextAvailable = AvailableFrames.Dequeue();
+			if (NextAvailable.IsSet())
+			{
+				return CaptureFrames[NextAvailable.GetValue()].Get();
+			}
+
+			return nullptr;
+		}
+
+		FCaptureFrame* PeekNextPending()
+		{
+			int32* NextPending = PendingFrames.Peek();
+			if (NextPending)
+			{
+				return CaptureFrames[*NextPending].Get();
+			}
+
+			return nullptr;
+		}
+
+	private:
+		TArray<TUniquePtr<FCaptureFrame>> CaptureFrames;
+		TSpscQueue<int32> AvailableFrames;
+		TSpscQueue<int32> PendingFrames;
+	};
 
 	/** Helper struct to contain arguments for CaptureFrame */
 	struct FCaptureFrameArgs
 	{
-		FRHICommandListImmediate& RHICmdList;
+		FRDGBuilder& GraphBuilder;
 		TObjectPtr<UMediaCapture> MediaCapture = nullptr;
-		FSceneViewport* CapturingSceneViewport = nullptr;
-		FTextureRenderTargetResource* TextureRenderTargetResource = nullptr;
+		FTexture2DRHIRef ResourceToCapture;
 		FIntPoint DesiredSize = FIntPoint::ZeroValue;
 	};
 
 	/** Helper struct to contain arguments for AddConversionPass */
 	struct FConversionPassArgs
 	{
-		FRDGBuilder& GraphBuilder;
 		const FRDGTextureRef& SourceRGBTexture;
 		bool bRequiresFormatConversion = false;
 		FRHICopyTextureInfo CopyInfo;
@@ -341,61 +421,61 @@ namespace UE::MediaCaptureData
 	{
 	public:
 
-		static FTexture2DRHIRef GetSourceTextureForInput(const FCaptureFrameArgs& Args)
+		static FTexture2DRHIRef GetSourceTextureForInput(FRHICommandListImmediate& RHICmdList, FSceneViewport* CapturingSceneViewport, FTextureRenderTargetResource* TextureRenderTargetResource)
 		{
 			FTexture2DRHIRef SourceTexture;
-			if (Args.CapturingSceneViewport)
+			if (CapturingSceneViewport)
 			{
 #if WITH_EDITOR
 				if (!IsRunningGame())
 				{
 					// PIE, PIE in windows, editor viewport
-					SourceTexture = Args.CapturingSceneViewport->GetRenderTargetTexture();
-					if (!SourceTexture.IsValid() && Args.CapturingSceneViewport->GetViewportRHI())
+					SourceTexture = CapturingSceneViewport->GetRenderTargetTexture();
+					if (!SourceTexture.IsValid() && CapturingSceneViewport->GetViewportRHI())
 					{
-						SourceTexture = Args.RHICmdList.GetViewportBackBuffer(Args.CapturingSceneViewport->GetViewportRHI());
+						SourceTexture = RHICmdList.GetViewportBackBuffer(CapturingSceneViewport->GetViewportRHI());
 					}
 				}
 				else
 #endif
-					if (Args.CapturingSceneViewport->GetViewportRHI())
+					if (CapturingSceneViewport->GetViewportRHI())
 					{
 						// Standalone and packaged
-						SourceTexture = Args.RHICmdList.GetViewportBackBuffer(Args.CapturingSceneViewport->GetViewportRHI());
+						SourceTexture = RHICmdList.GetViewportBackBuffer(CapturingSceneViewport->GetViewportRHI());
 					}
 			}
-			else if (Args.TextureRenderTargetResource && Args.TextureRenderTargetResource->GetTextureRenderTarget2DResource())
+			else if (TextureRenderTargetResource && TextureRenderTargetResource->GetTextureRenderTarget2DResource())
 			{
-				SourceTexture = Args.TextureRenderTargetResource->GetTextureRenderTarget2DResource()->GetTextureRHI();
+				SourceTexture = TextureRenderTargetResource->GetTextureRenderTarget2DResource()->GetTextureRHI();
 			}
 
 			return SourceTexture;
 		}
 
-		static bool AreInputsValid(const FCaptureFrameArgs& Args, const FTexture2DRHIRef& SourceTexture)
+		static bool AreInputsValid(const FCaptureFrameArgs& Args)
 		{
 			// If it is a simple rgba swizzle we can handle the conversion. Supported formats
 			// contained in SupportedRgbaSwizzleFormats. Warning would've been displayed on start of capture.
-			if (Args.MediaCapture->DesiredPixelFormat != SourceTexture->GetFormat() &&
-				(!SupportedRgbaSwizzleFormats.Contains(SourceTexture->GetFormat()) || !Args.MediaCapture->DesiredCaptureOptions.bConvertToDesiredPixelFormat))
+			if (Args.MediaCapture->DesiredPixelFormat != Args.ResourceToCapture->GetFormat() &&
+				(!SupportedRgbaSwizzleFormats.Contains(Args.ResourceToCapture->GetFormat()) || !Args.MediaCapture->DesiredCaptureOptions.bConvertToDesiredPixelFormat))
 			{
 				UE_LOG(LogMediaIOCore, Error, TEXT("The capture will stop for '%s'. The Source pixel format doesn't match with the user requested pixel format. %sRequested: %s Source: %s")
 					, *Args.MediaCapture->MediaOutputName
-					, (SupportedRgbaSwizzleFormats.Contains(SourceTexture->GetFormat()) && !Args.MediaCapture->DesiredCaptureOptions.bConvertToDesiredPixelFormat) ? TEXT("Please enable \"Convert To Desired Pixel Format\" option in Media Capture settings. ") : TEXT("")
+					, (SupportedRgbaSwizzleFormats.Contains(Args.ResourceToCapture->GetFormat()) && !Args.MediaCapture->DesiredCaptureOptions.bConvertToDesiredPixelFormat) ? TEXT("Please enable \"Convert To Desired Pixel Format\" option in Media Capture settings. ") : TEXT("")
 					, GetPixelFormatString(Args.MediaCapture->DesiredPixelFormat)
-					, GetPixelFormatString(SourceTexture->GetFormat()));
+					, GetPixelFormatString(Args.ResourceToCapture->GetFormat()));
 
 				return false;
 			}
 
 			if (Args.MediaCapture->DesiredCaptureOptions.Crop == EMediaCaptureCroppingType::None)
 			{
-				if (Args.DesiredSize.X != SourceTexture->GetSizeX() || Args.DesiredSize.Y != SourceTexture->GetSizeY())
+				if (Args.DesiredSize.X != Args.ResourceToCapture->GetSizeX() || Args.DesiredSize.Y != Args.ResourceToCapture->GetSizeY())
 				{
 					UE_LOG(LogMediaIOCore, Error, TEXT("The capture will stop for '%s'. The Source size doesn't match with the user requested size. Requested: %d,%d  Source: %d,%d")
 						, *Args.MediaCapture->MediaOutputName
 						, Args.DesiredSize.X, Args.DesiredSize.Y
-						, SourceTexture->GetSizeX(), SourceTexture->GetSizeY());
+						, Args.ResourceToCapture->GetSizeX(), Args.ResourceToCapture->GetSizeY());
 
 					return false;
 				}
@@ -408,12 +488,12 @@ namespace UE::MediaCaptureData
 					StartCapturePoint = Args.MediaCapture->DesiredCaptureOptions.CustomCapturePoint;
 				}
 
-				if ((uint32)(Args.DesiredSize.X + StartCapturePoint.X) > SourceTexture->GetSizeX() || (uint32)(Args.DesiredSize.Y + StartCapturePoint.Y) > SourceTexture->GetSizeY())
+				if ((uint32)(Args.DesiredSize.X + StartCapturePoint.X) > Args.ResourceToCapture->GetSizeX() || (uint32)(Args.DesiredSize.Y + StartCapturePoint.Y) > Args.ResourceToCapture->GetSizeY())
 				{
 					UE_LOG(LogMediaIOCore, Error, TEXT("The capture will stop for '%s'. The Source size doesn't match with the user requested size. Requested: %d,%d  Source: %d,%d")
 						, *Args.MediaCapture->MediaOutputName
 						, Args.DesiredSize.X, Args.DesiredSize.Y
-						, SourceTexture->GetSizeX(), SourceTexture->GetSizeY());
+						, Args.ResourceToCapture->GetSizeX(), Args.ResourceToCapture->GetSizeY());
 
 					return false;
 				}
@@ -472,7 +552,7 @@ namespace UE::MediaCaptureData
 			// If no conversion was required, go through a simple copy
 			if (Args.MediaCapture->ConversionOperation == EMediaCaptureConversionOperation::NONE && !bRequiresFormatConversion)
 			{
-				AddCopyTexturePass(ConversionPassArgs.GraphBuilder, ConversionPassArgs.SourceRGBTexture, OutputResource, ConversionPassArgs.CopyInfo);
+				AddCopyTexturePass(Args.GraphBuilder, ConversionPassArgs.SourceRGBTexture, OutputResource, ConversionPassArgs.CopyInfo);
 			}
 			else
 			{
@@ -490,8 +570,8 @@ namespace UE::MediaCaptureData
 					FScreenPassTextureViewport InputViewport(ConversionPassArgs.SourceRGBTexture, ViewRect);
 					FScreenPassTextureViewport OutputViewport(OutputResource);
 					TShaderMapRef<FRGB8toUYVY8ConvertPS> PixelShader(GlobalShaderMap);
-					FRGB8toUYVY8ConvertPS::FParameters* Parameters = PixelShader->AllocateAndSetParameters(ConversionPassArgs.GraphBuilder, ConversionPassArgs.SourceRGBTexture, MediaShaders::RgbToYuvRec709Scaled, MediaShaders::YUVOffset8bits, bDoLinearToSRGB, OutputResource);
-					AddDrawScreenPass(ConversionPassArgs.GraphBuilder, RDG_EVENT_NAME("RGBToUYVY 8 bit"), ViewInfo, OutputViewport, InputViewport, VertexShader, PixelShader, Parameters);
+					FRGB8toUYVY8ConvertPS::FParameters* Parameters = PixelShader->AllocateAndSetParameters(Args.GraphBuilder, ConversionPassArgs.SourceRGBTexture, MediaShaders::RgbToYuvRec709Scaled, MediaShaders::YUVOffset8bits, bDoLinearToSRGB, OutputResource);
+					AddDrawScreenPass(Args.GraphBuilder, RDG_EVENT_NAME("RGBToUYVY 8 bit"), ViewInfo, OutputViewport, InputViewport, VertexShader, PixelShader, Parameters);
 				}
 				break;
 				case EMediaCaptureConversionOperation::RGB10_TO_YUVv210_10BIT:
@@ -501,8 +581,8 @@ namespace UE::MediaCaptureData
 					FScreenPassTextureViewport InputViewport(ConversionPassArgs.SourceRGBTexture, ViewRect);
 					FScreenPassTextureViewport OutputViewport(OutputResource);
 					TShaderMapRef<FRGB10toYUVv210ConvertPS> PixelShader(GlobalShaderMap);
-					FRGB10toYUVv210ConvertPS::FParameters* Parameters = PixelShader->AllocateAndSetParameters(ConversionPassArgs.GraphBuilder, ConversionPassArgs.SourceRGBTexture, MediaShaders::RgbToYuvRec709Scaled, MediaShaders::YUVOffset10bits, bDoLinearToSRGB, OutputResource);
-					AddDrawScreenPass(ConversionPassArgs.GraphBuilder, RDG_EVENT_NAME("RGBToYUVv210"), ViewInfo, OutputViewport, InputViewport, VertexShader, PixelShader, Parameters);
+					FRGB10toYUVv210ConvertPS::FParameters* Parameters = PixelShader->AllocateAndSetParameters(Args.GraphBuilder, ConversionPassArgs.SourceRGBTexture, MediaShaders::RgbToYuvRec709Scaled, MediaShaders::YUVOffset10bits, bDoLinearToSRGB, OutputResource);
+					AddDrawScreenPass(Args.GraphBuilder, RDG_EVENT_NAME("RGBToYUVv210"), ViewInfo, OutputViewport, InputViewport, VertexShader, PixelShader, Parameters);
 				}
 				break;
 				case EMediaCaptureConversionOperation::INVERT_ALPHA:
@@ -524,8 +604,8 @@ namespace UE::MediaCaptureData
 						PermutationVector.Set<FModifyAlphaSwizzleRgbaPS::FConversionOp>(static_cast<int32>(MediaConversionOperation));
 
 						TShaderMapRef<FModifyAlphaSwizzleRgbaPS> PixelShader(GlobalShaderMap, PermutationVector);
-						FModifyAlphaSwizzleRgbaPS::FParameters* Parameters = PixelShader->AllocateAndSetParameters(ConversionPassArgs.GraphBuilder, ConversionPassArgs.SourceRGBTexture, OutputResource);
-						AddDrawScreenPass(ConversionPassArgs.GraphBuilder, RDG_EVENT_NAME("MediaCaptureSwizzle"), ViewInfo, OutputViewport, InputViewport, VertexShader, PixelShader, Parameters);
+						FModifyAlphaSwizzleRgbaPS::FParameters* Parameters = PixelShader->AllocateAndSetParameters(Args.GraphBuilder, ConversionPassArgs.SourceRGBTexture, OutputResource);
+						AddDrawScreenPass(Args.GraphBuilder, RDG_EVENT_NAME("MediaCaptureSwizzle"), ViewInfo, OutputViewport, InputViewport, VertexShader, PixelShader, Parameters);
 					}
 					break;
 				}
@@ -547,20 +627,17 @@ namespace UE::MediaCaptureData
 				return false;
 			}
 
-			// Get source texture that is going to be captured
-			const FTexture2DRHIRef SourceTexture = GetSourceTextureForInput(Args);
-
 			// If true, we will need to go through our different shader to convert from source format to out format (i.e RGB to YUV)
-			const bool bRequiresFormatConversion = Args.MediaCapture->DesiredPixelFormat != SourceTexture->GetFormat();
+			const bool bRequiresFormatConversion = Args.MediaCapture->DesiredPixelFormat != Args.ResourceToCapture->GetFormat();
 
-			if (!SourceTexture.IsValid())
+			if (!Args.ResourceToCapture.IsValid())
 			{
 				UE_LOG(LogMediaIOCore, Error, TEXT("Can't grab the Texture to capture for '%s'."), *Args.MediaCapture->MediaOutputName);
 				return false;
 			}
 
 			// Validate pixel formats and sizes before pursuing
-			if (AreInputsValid(Args, SourceTexture) == false)
+			if (AreInputsValid(Args) == false)
 			{
 				return false;
 			}
@@ -581,50 +658,50 @@ namespace UE::MediaCaptureData
 				FRHICopyTextureInfo CopyInfo;
 				FVector2D SizeU;
 				FVector2D SizeV;
-				GetCopyInfo(Args, SourceTexture, CopyInfo, SizeU, SizeV);
-
-				FRDGBuilder GraphBuilder(Args.RHICmdList);
+				GetCopyInfo(Args, Args.ResourceToCapture, CopyInfo, SizeU, SizeV);
 
 				// Register the source texture that we want to capture
-				const FRDGTextureRef SourceRGBTexture = GraphBuilder.RegisterExternalTexture(CreateRenderTarget(SourceTexture, TEXT("SourceTexture")));
+				const FRDGTextureRef SourceRGBTexture = Args.GraphBuilder.RegisterExternalTexture(CreateRenderTarget(Args.ResourceToCapture, TEXT("SourceTexture")));
 
 				// Register output resource used by the current capture method (texture or buffer)
-				typename TFrameType::FOutputResourceType OutputResource = CapturingFrame->RegisterResource(GraphBuilder);
+				typename TFrameType::FOutputResourceType OutputResource = CapturingFrame->RegisterResource(Args.GraphBuilder);
 				{
 					TRACE_CPUPROFILER_EVENT_SCOPE(UMediaCapture::GraphSetup);
-					SCOPED_DRAW_EVENTF(Args.RHICmdList, MediaCapture, TEXT("MediaCapture"));
+					SCOPED_DRAW_EVENTF(Args.GraphBuilder.RHICmdList, MediaCapture, TEXT("MediaCapture"));
 
 					// If custom conversion was requested from implementation, give it useful information to apply 
 					if (Args.MediaCapture->ConversionOperation == EMediaCaptureConversionOperation::CUSTOM)
 					{
 						TRACE_CPUPROFILER_EVENT_SCOPE(UMediaCapture::CustomCapture);
 
-						Args.MediaCapture->OnCustomCapture_RenderingThread(GraphBuilder, CapturingFrame->CaptureBaseData, CapturingFrame->UserData
+						Args.MediaCapture->OnCustomCapture_RenderingThread(Args.GraphBuilder, CapturingFrame->CaptureBaseData, CapturingFrame->UserData
 							, SourceRGBTexture, OutputResource, CopyInfo, SizeU, SizeV);
 					}
 					else
 					{
 						TRACE_CPUPROFILER_EVENT_SCOPE(UMediaCapture::FormatConversion);
-						AddConversionPass(Args, { GraphBuilder, SourceRGBTexture, bRequiresFormatConversion,  CopyInfo, SizeU, SizeV }, OutputResource);
+						AddConversionPass(Args, { SourceRGBTexture, bRequiresFormatConversion,  CopyInfo, SizeU, SizeV }, OutputResource);
 					}
 				}
-
+				
 				// If Capture implementation is not grabbing GPU resource directly, push a readback pass to access it from CPU
 				if (Args.MediaCapture->bShouldCaptureRHIResource == false)
 				{
 					TRACE_CPUPROFILER_EVENT_SCOPE(UMediaCapture::EnqueueReadback);
 
-					CapturingFrame->EnqueueCopy(GraphBuilder, OutputResource);
+					CapturingFrame->EnqueueCopy(Args.GraphBuilder, OutputResource);
 
 					CapturingFrame->bReadbackRequested = true;
-					const int32 FrameIndex = Args.MediaCapture->CaptureFrames.IndexOfByPredicate([CapturingFrame](const TUniquePtr<FCaptureFrame>& InCaptureFrame) { return InCaptureFrame.Get() == CapturingFrame; });
-					UE_LOG(LogMediaIOCore, VeryVerbose, TEXT("Requested copy for capture frame %d"), FrameIndex);
-					++Args.MediaCapture->NumberOfResourcesToReadback;
+					UE_LOG(LogMediaIOCore, VeryVerbose, TEXT("Requested copy for capture frame %d"), CapturingFrame->FrameId);
+				}
+				else
+				{
+					CapturingFrame->bDoingGPUCopy = true;
 				}
 
-				CapturingFrame->bDoingGPUCopy = true;
+				Args.MediaCapture->FrameManager->MarkPending(*CapturingFrame);
 
-				GraphBuilder.Execute();
+				++Args.MediaCapture->PendingFrameCount;
 			}
 
 			return true;
@@ -663,10 +740,11 @@ FMediaCaptureOptions::FMediaCaptureOptions()
 UMediaCapture::UMediaCapture(const FObjectInitializer& ObjectInitializer)
 	: Super(ObjectInitializer)
 	, ValidSourceGPUMask(FRHIGPUMask::All())
+	, FrameManager(MakeUnique<UE::MediaCaptureData::FFrameManager>())
 	, bOutputResourcesInitialized(false)
 	, bShouldCaptureRHIResource(false)
 	, WaitingForRenderCommandExecutionCounter(0)
-	, NumberOfResourcesToReadback(0)
+	, PendingFrameCount(0)
 {
 }
 
@@ -750,9 +828,14 @@ bool UMediaCapture::CaptureSceneViewport(TSharedPtr<FSceneViewport>& InSceneView
 	}
 
 	SetState(EMediaCaptureState::Preparing);
-	bool bInitialized = CaptureSceneViewportImpl(InSceneViewport);
 
-	// This could have been updated by the call to CaptureSceneViewportImpl
+	bool bInitialized = InitializeCapture();
+	if (bInitialized)
+	{
+		bInitialized = PostInitializeCaptureViewport(InSceneViewport);
+	}
+
+	// This could have been updated by the initialization done by the implementation
 	bShouldCaptureRHIResource = ShouldCaptureRHIResource(); 
 	if (bShouldCaptureRHIResource && DesiredOutputResourceType != EMediaCaptureResourceType::Texture)
 	{
@@ -772,8 +855,6 @@ bool UMediaCapture::CaptureSceneViewport(TSharedPtr<FSceneViewport>& InSceneView
 	{
 		//no lock required, the command on the render thread is not active
 		CapturingSceneViewport = InSceneViewport;
-
-		CurrentReadbackIndex = 0;
 		FCoreDelegates::OnEndFrame.AddUObject(this, &UMediaCapture::OnEndFrame_GameThread);
 	}
 	else
@@ -784,12 +865,78 @@ bool UMediaCapture::CaptureSceneViewport(TSharedPtr<FSceneViewport>& InSceneView
 	}
 
 #if WITH_EDITOR
-	MediaCaptureAnalytics::SendCaptureEvent(TEXT("SceneViewport"));
+	MediaCaptureAnalytics::SendCaptureEvent(GetCaptureSourceType());
 #endif
 	
 	return bInitialized;
 }
 
+bool UMediaCapture::CaptureRHITexture(const FRHICaptureResourceDescription& ResourceDescription, FMediaCaptureOptions CaptureOptions)
+{
+	StopCapture(false);
+
+	check(IsInGameThread());
+
+	DesiredCaptureOptions = CaptureOptions;
+
+	if (!ValidateMediaOutput())
+	{
+		MediaCaptureDetails::ShowSlateNotification();
+		return false;
+	}
+
+	CacheMediaOutput(EMediaCaptureSourceType::RHI_RESOURCE);
+
+	// If output is configured to capture resource size directly (not enforced), configure desired size to be the same as resource
+	if (bUseRequestedTargetSize)
+	{
+		DesiredSize = ResourceDescription.ResourceSize;
+	}
+
+	CacheOutputOptions();
+
+	// Can't validate resource to capture as it will be given everytime a capture is requested
+	SetState(EMediaCaptureState::Preparing);
+
+	bool bInitialized = InitializeCapture();
+	if (bInitialized)
+	{
+		bInitialized = PostInitializeCaptureRHIResource(ResourceDescription);
+	}
+
+	// This could have been updated by the call to implementation initialization
+	bShouldCaptureRHIResource = ShouldCaptureRHIResource();
+	if (bShouldCaptureRHIResource && DesiredOutputResourceType != EMediaCaptureResourceType::Texture)
+	{
+		UE_LOG(LogMediaIOCore, Warning, TEXT("Can't capture RHI resource when not using texture output resource"));
+		SetState(EMediaCaptureState::Stopped);
+		MediaCaptureDetails::ShowSlateNotification();
+	}
+
+	if (bInitialized)
+	{
+		InitializeOutputResources(MediaOutput->NumberOfTextureBuffers);
+		bInitialized = GetState() != EMediaCaptureState::Stopped;
+	}
+
+	if (bInitialized)
+	{
+		// Used to queue capture data production from the game thread to be consumed when frame capture is requested on render thread
+		FCoreDelegates::OnBeginFrame.AddUObject(this, &UMediaCapture::OnBeginFrame_GameThread);
+	}
+	else
+	{
+		SetState(EMediaCaptureState::Stopped);
+		MediaCaptureDetails::ShowSlateNotification();
+	}
+
+#if WITH_EDITOR
+	MediaCaptureAnalytics::SendCaptureEvent(GetCaptureSourceType());
+#endif
+
+	return bInitialized;
+}
+
 bool UMediaCapture::CaptureTextureRenderTarget2D(UTextureRenderTarget2D* InRenderTarget2D, FMediaCaptureOptions CaptureOptions)
 {
 	StopCapture(false);
@@ -825,10 +972,21 @@ bool UMediaCapture::CaptureTextureRenderTarget2D(UTextureRenderTarget2D* InRende
 	}
 
 	SetState(EMediaCaptureState::Preparing);
-	bool bInitialized = CaptureRenderTargetImpl(InRenderTarget2D);
 
-	// This could have been updated by the call to CaptureSceneViewportImpl
+	bool bInitialized = InitializeCapture();
+	if (bInitialized)
+	{
+		bInitialized = PostInitializeCaptureRenderTarget(InRenderTarget2D);
+	}
+
+	// This could have been updated by the initialization done by the implementation
 	bShouldCaptureRHIResource = ShouldCaptureRHIResource();
+	if (bShouldCaptureRHIResource && DesiredOutputResourceType != EMediaCaptureResourceType::Texture)
+	{
+		UE_LOG(LogMediaIOCore, Warning, TEXT("Can't capture RHI resource when not using texture output resource"));
+		SetState(EMediaCaptureState::Stopped);
+		MediaCaptureDetails::ShowSlateNotification();
+	}
 
 	if (bInitialized)
 	{
@@ -841,7 +999,6 @@ bool UMediaCapture::CaptureTextureRenderTarget2D(UTextureRenderTarget2D* InRende
 		//no lock required the command on the render thread is not active yet
 		CapturingRenderTarget = InRenderTarget2D;
 
-		CurrentReadbackIndex = 0;
 		FCoreDelegates::OnEndFrame.AddUObject(this, &UMediaCapture::OnEndFrame_GameThread);
 	}
 	else
@@ -851,7 +1008,7 @@ bool UMediaCapture::CaptureTextureRenderTarget2D(UTextureRenderTarget2D* InRende
 	}
 
 #if WITH_EDITOR
-	MediaCaptureAnalytics::SendCaptureEvent(TEXT("RenderTarget2D"));
+	MediaCaptureAnalytics::SendCaptureEvent(GetCaptureSourceType());
 #endif
 	
 	return bInitialized;
@@ -860,6 +1017,7 @@ bool UMediaCapture::CaptureTextureRenderTarget2D(UTextureRenderTarget2D* InRende
 void UMediaCapture::CacheMediaOutput(EMediaCaptureSourceType InSourceType)
 {
 	check(MediaOutput);
+	CaptureSourceType = InSourceType;
 	DesiredSize = MediaOutput->GetRequestedSize();
 	bUseRequestedTargetSize = DesiredSize == UMediaOutput::RequestCaptureSourceSize;
 	DesiredPixelFormat = MediaOutput->GetRequestedPixelFormat();
@@ -965,7 +1123,7 @@ bool UMediaCapture::UpdateSceneViewport(TSharedPtr<FSceneViewport>& InSceneViewp
 
 	{
 		FScopeLock Lock(&AccessingCapturingSource);
-		while(WaitingForRenderCommandExecutionCounter.Load() > 0)
+		while(WaitingForRenderCommandExecutionCounter.load() > 0)
 		{
 			FlushRenderingCommands();
 		}
@@ -1010,7 +1168,7 @@ bool UMediaCapture::UpdateTextureRenderTarget2D(UTextureRenderTarget2D * InRende
 
 	{
 		FScopeLock Lock(&AccessingCapturingSource);
-		while (WaitingForRenderCommandExecutionCounter.Load() > 0)
+		while (WaitingForRenderCommandExecutionCounter.load() > 0)
 		{
 			FlushRenderingCommands();
 		}
@@ -1040,7 +1198,7 @@ void UMediaCapture::StopCapture(bool bAllowPendingFrameToBeProcess)
 			//Do not flush when auto stopping to avoid hitches.
 			if(DesiredCaptureOptions.bAutostopOnCapture != true)
 			{
-			while (WaitingForRenderCommandExecutionCounter.Load() > 0)
+			while (WaitingForRenderCommandExecutionCounter.load() > 0)
 			{
 				FlushRenderingCommands();
 			}
@@ -1053,9 +1211,10 @@ void UMediaCapture::StopCapture(bool bAllowPendingFrameToBeProcess)
 		{
 			SetState(EMediaCaptureState::Stopped);
 
+			FCoreDelegates::OnBeginFrame.RemoveAll(this);
 			FCoreDelegates::OnEndFrame.RemoveAll(this);
 
-			while (WaitingForRenderCommandExecutionCounter.Load() > 0 || !bOutputResourcesInitialized)
+			while (WaitingForRenderCommandExecutionCounter.load() > 0 || !bOutputResourcesInitialized)
 			{
 				FlushRenderingCommands();
 			}
@@ -1072,12 +1231,12 @@ void UMediaCapture::StopCapture(bool bAllowPendingFrameToBeProcess)
 			ConversionOperation = EMediaCaptureConversionOperation::NONE;
 			MediaOutputName.Reset();
 
-			// CaptureFrames contains FTexture2DRHIRef, therefore should be released on Render Tread thread.
+			// CaptureFrames contains RHI resources, therefore should be released on Render Tread thread.
 			// Keep references frames to be released in a temporary array and clear CaptureFrames on Game Thread.
 			ENQUEUE_RENDER_COMMAND(MediaOutputReleaseCaptureFrames)(
-				[TempArrayToBeReleasedOnRenderThread = MoveTemp(CaptureFrames)](FRHICommandListImmediate& RHICmdList) mutable
+				[FrameManagerToRelease = MoveTemp(FrameManager)](FRHICommandListImmediate& RHICmdList) mutable
 				{
-					TempArrayToBeReleasedOnRenderThread.Reset();
+					FrameManagerToRelease.Reset();
 				});
 		}
 	}
@@ -1091,6 +1250,150 @@ void UMediaCapture::SetMediaOutput(UMediaOutput* InMediaOutput)
 	}
 }
 
+void UMediaCapture::CaptureImmediate_RenderThread(FRDGBuilder& GraphBuilder, FRHITexture* InSourceTexture)
+{
+	using namespace UE::MediaCaptureData;
+	TRACE_CPUPROFILER_EVENT_SCOPE(UMediaCapture::CaptureImmediate_RenderThread);
+	
+	check(IsInRenderingThread());
+
+	// This could happen if the capture immediate is called before our resources were initialized. We can't really flush as we are in a command.
+	if (!bOutputResourcesInitialized)
+	{
+		UE_LOG(LogMediaIOCore, VeryVerbose, TEXT("Could not capture frame. Output resources haven't been initialized"));
+		return;
+	}
+
+	if (!MediaOutput)
+	{
+		return;
+	}
+
+	if (GetState() == EMediaCaptureState::Error)
+	{
+		return;
+	}
+
+	if (CaptureSourceType != EMediaCaptureSourceType::RHI_RESOURCE)
+	{
+		UE_LOG(LogMediaIOCore, Error, TEXT("[%s] - Trying to capture a RHI resource with another capture type."), *MediaOutputName);
+		SetState(EMediaCaptureState::Error);
+	}
+
+	if (GetState() != EMediaCaptureState::Capturing && GetState() != EMediaCaptureState::StopRequested)
+	{
+		return;
+	}
+
+	if (DesiredCaptureOptions.bSkipFrameWhenRunningExpensiveTasks && !FSlateThrottleManager::Get().IsAllowingExpensiveTasks())
+	{
+		return;
+	}
+
+	// Get cached capture data from game thread. We want to find a cached frame matching current render thread frame number
+	bool bFoundMatchingData = false;
+	FQueuedCaptureData NextCaptureData;
+	while (CaptureDataQueue.Dequeue(NextCaptureData))
+	{
+		if (NextCaptureData.BaseData.SourceFrameNumberRenderThread == GFrameNumberRenderThread)
+		{
+			bFoundMatchingData = true;
+			break;
+		}
+	}
+
+	if (bFoundMatchingData == false)
+	{
+		UE_LOG(LogMediaIOCore, Warning, TEXT("Can't capture frame. Could not find a matching game frame."));
+		return;
+	}
+
+	if (GetState() == EMediaCaptureState::StopRequested && PendingFrameCount.load() <= 0)
+	{
+		// All the requested frames have been captured.
+		StopCapture(false);
+		return;
+	}
+		
+	// Get next available frame from the store. Can be invalid.
+	FCaptureFrame* CapturingFrame = nullptr;
+	if (GetState() != EMediaCaptureState::StopRequested)
+	{
+		CapturingFrame = FrameManager->GetNextAvailable();
+	}
+
+	if (CapturingFrame == nullptr && GetState() != EMediaCaptureState::StopRequested)
+	{
+		if (DesiredCaptureOptions.OverrunAction == EMediaCaptureOverrunAction::Flush)
+		{
+			UE_LOG(LogMediaIOCore, Error, TEXT("[%s] - No frames available for capture. This should not happen."), *MediaOutputName);
+			SetState(EMediaCaptureState::Error);
+			return;
+		}
+		else
+		{
+			//In case we are skipping frames, just keep capture frame as invalid
+			UE_LOG(LogMediaIOCore, Verbose, TEXT("[%s] - No frames available for capture. Skipping"), *MediaOutputName);
+		}
+	}
+
+	// Peek next pending before capturing to avoid getting the newly captured one. 
+	// It shouldn't happen but if it happens it will start waiting for the capture to complete
+	FCaptureFrame* NextPending = FrameManager->PeekNextPending();
+	if (NextPending && CapturingFrame)
+	{
+		ensure(NextPending->FrameId != CapturingFrame->FrameId);
+	}
+
+	PrintFrameState();
+
+	UE_LOG(LogMediaIOCore, Verbose, TEXT("[%s] - Capturing frame %d"), *MediaOutputName, CapturingFrame ? CapturingFrame->FrameId : -1);
+	if (CapturingFrame)
+	{
+		// Prepare frame to capture
+
+		// Use queued capture base data for this frame to capture
+		CapturingFrame->CaptureBaseData = MoveTemp(NextCaptureData.BaseData);
+		CapturingFrame->CaptureBaseData.SourceFrameNumber = ++CaptureRequestCount;
+		CapturingFrame->UserData = MoveTemp(NextCaptureData.UserData);
+
+		ProcessCapture_RenderThread(GraphBuilder, this, CapturingFrame, InSourceTexture, DesiredSize);
+	}
+
+	UE_LOG(LogMediaIOCore, Verbose, TEXT("[%s] - Processing pending frame %d"), *MediaOutputName, NextPending ? NextPending->FrameId : -1);
+	if(NextPending)
+	{
+		ProcessReadyFrame_RenderThread(GraphBuilder.RHICmdList, this, NextPending);
+	}
+}
+
+const FString& UMediaCapture::GetCaptureSourceType() const
+{
+	switch (CaptureSourceType)
+	{
+		case EMediaCaptureSourceType::SCENE_VIEWPORT:
+		{
+			static FString Viewport = TEXT("SceneViewport");
+			return Viewport;
+		}
+		case EMediaCaptureSourceType::RENDER_TARGET:
+		{
+			static FString RenderTarget = TEXT("RenderTarget");
+			return RenderTarget;
+		}
+		case EMediaCaptureSourceType::RHI_RESOURCE:
+		{
+			static FString RHIResource = TEXT("RHI");
+			return RHIResource;
+		}
+		default:
+		{
+			static FString Invalid = TEXT("Invalid");
+			return Invalid;
+		}
+	}
+}
+
 void UMediaCapture::SetState(EMediaCaptureState InNewState)
 {
 	if (MediaState != InNewState)
@@ -1131,7 +1434,7 @@ void UMediaCapture::ResetFixedViewportSize(TSharedPtr<FSceneViewport> InViewport
 {
 	if (bViewportHasFixedViewportSize && InViewport.IsValid())
 	{
-		if (bInFlushRenderingCommands && WaitingForRenderCommandExecutionCounter.Load() > 0)
+		if (bInFlushRenderingCommands && WaitingForRenderCommandExecutionCounter.load() > 0)
 		{
 			FlushRenderingCommands();
 		}
@@ -1142,7 +1445,7 @@ void UMediaCapture::ResetFixedViewportSize(TSharedPtr<FSceneViewport> InViewport
 
 bool UMediaCapture::HasFinishedProcessing() const
 {
-	return WaitingForRenderCommandExecutionCounter.Load() == 0
+	return WaitingForRenderCommandExecutionCounter.load() == 0
 		|| GetState() == EMediaCaptureState::Error
 		|| GetState() == EMediaCaptureState::Stopped;
 }
@@ -1163,18 +1466,21 @@ void UMediaCapture::InitializeOutputResources(int32 InNumberOfBuffers)
 		return;
 	}
 
+	bOutputResourcesInitialized = false;
 	NumberOfCaptureFrame = InNumberOfBuffers;
-	check(CaptureFrames.Num() == 0);
+
+	check(NumberOfCaptureFrame >= 2);
+	check(FrameManager->GetFrames().Num() == 0);
 
 	UMediaCapture* This = this;
 	ENQUEUE_RENDER_COMMAND(MediaOutputCaptureFrameCreateResources)(
-		[This](FRHICommandListImmediate& RHICmdList)
+	[This](FRHICommandListImmediate& RHICmdList)
 		{
 			for (int32 Index = 0; Index < This->NumberOfCaptureFrame; ++Index)
 			{
 				if (This->DesiredOutputResourceType == EMediaCaptureResourceType::Texture)
 				{
-					TUniquePtr<FTextureCaptureFrame> NewFrame = MakeUnique<FTextureCaptureFrame>();
+					TUniquePtr<FTextureCaptureFrame> NewFrame = MakeUnique<FTextureCaptureFrame>(Index);
 					FRDGTextureDesc OutputDesc = FRDGTextureDesc::Create2D(
 						This->DesiredOutputSize,
 						This->DesiredOutputPixelFormat,
@@ -1190,11 +1496,11 @@ void UMediaCapture::InitializeOutputResources(int32 InNumberOfBuffers)
 						NewFrame->ReadbackTexture = MakeUnique<FRHIGPUTextureReadback>(*FString::Printf(TEXT("MediaCaptureTextureReadback_%d"), Index));
 					}
 
-					This->CaptureFrames.Emplace(MoveTemp(NewFrame));
+					This->FrameManager->AddFrame(MoveTemp(NewFrame));
 				}
 				else
 				{
-					TUniquePtr<FBufferCaptureFrame> NewFrame = MakeUnique<FBufferCaptureFrame>();
+					TUniquePtr<FBufferCaptureFrame> NewFrame = MakeUnique<FBufferCaptureFrame>(Index);
 
 					if (This->DesiredOutputBufferDescription.NumElements > 0)
 					{
@@ -1206,7 +1512,7 @@ void UMediaCapture::InitializeOutputResources(int32 InNumberOfBuffers)
 							NewFrame->ReadbackBuffer = MakeUnique<FRHIGPUBufferReadback>(*FString::Printf(TEXT("MediaCaptureBufferReadback_%d"), Index));
 						}
 
-						This->CaptureFrames.Emplace(MoveTemp(NewFrame));
+						This->FrameManager->AddFrame(MoveTemp(NewFrame));
 					}
 					else
 					{
@@ -1243,12 +1549,36 @@ bool UMediaCapture::ValidateMediaOutput() const
 	return true;
 }
 
+bool UMediaCapture::CaptureSceneViewportImpl(TSharedPtr<FSceneViewport>&InSceneViewport)
+{
+	return false;
+}
+
+bool UMediaCapture::CaptureRenderTargetImpl(UTextureRenderTarget2D * InRenderTarget)
+{
+	return false;
+}
+
+void UMediaCapture::OnBeginFrame_GameThread()
+{
+	
+	if (CaptureSourceType == EMediaCaptureSourceType::RHI_RESOURCE && (GetState() == EMediaCaptureState::Capturing))
+	{
+		// Queue capture data to be consumed when capture requests are done on render thread
+		FQueuedCaptureData CaptureData;
+		CaptureData.BaseData.SourceFrameTimecode = FApp::GetTimecode();
+		CaptureData.BaseData.SourceFrameTimecodeFramerate = FApp::GetTimecodeFrameRate();
+		CaptureData.BaseData.SourceFrameNumberRenderThread = GFrameNumber;
+		CaptureData.UserData = GetCaptureFrameUserData_GameThread();
+		CaptureDataQueue.Enqueue(MoveTemp(CaptureData));
+	}
+}
+
 void UMediaCapture::OnEndFrame_GameThread()
 {
 	using namespace UE::MediaCaptureData;
 
-	const FString EndFrameTraceName = FString::Format(TEXT("MediaCapture::OnEndFrame_GameThread (Index {0})"), { CurrentReadbackIndex });
-	TRACE_CPUPROFILER_EVENT_SCOPE_TEXT(*EndFrameTraceName);
+	TRACE_CPUPROFILER_EVENT_SCOPE(MediaCapture::OnEndFrame_GameThread);
 
 	if (!bOutputResourcesInitialized)
 	{
@@ -1275,51 +1605,45 @@ void UMediaCapture::OnEndFrame_GameThread()
 		return;
 	}
 
-	CurrentReadbackIndex = (CurrentReadbackIndex + 1) % NumberOfCaptureFrame;
-	int32 ReadyFrameIndex = (CurrentReadbackIndex - 1); // Previous one in the buffer queue
-	if (ReadyFrameIndex < 0)
-	{
-		ReadyFrameIndex = NumberOfCaptureFrame - 1;
-	}
-
-	// Frame that should be on the system ram and we want to send to the user
-	FCaptureFrame* ReadyFrame = CaptureFrames[ReadyFrameIndex].Get();
-	// Frame that we want to transfer to system ram
-	FCaptureFrame* CapturingFrame = (GetState() != EMediaCaptureState::StopRequested) ? CaptureFrames[CurrentReadbackIndex].Get() : nullptr;
-
-	if (!ShouldCaptureRHIResource())
-	{
-		TStringBuilder<256> ReadbackInfoBuilder;
-		ReadbackInfoBuilder << "\n";
-		for (int32 Index = 0; Index < NumberOfCaptureFrame; Index++)
-		{
-			ReadbackInfoBuilder << FString::Format(TEXT("Frame {0} readback requested: {1}\n"), { Index, CaptureFrames[Index]->bReadbackRequested ? TEXT("true") : TEXT("false") });
-		}
-		UE_LOG(LogMediaIOCore, VeryVerbose, TEXT("%s"), ReadbackInfoBuilder.GetData());
-	}
-
-
-	UE_LOG(LogMediaIOCore, VeryVerbose, TEXT("MediaOutput: '%s'. ReadyFrameIndex: '%d' '%s'. CurrentReadbackIndex: '%d'.")
-		, *MediaOutputName, ReadyFrameIndex, (CaptureFrames[ReadyFrameIndex]->bReadbackRequested) ? TEXT("Y"): TEXT("N"), CurrentReadbackIndex);
-	if (GetState() == EMediaCaptureState::StopRequested && NumberOfResourcesToReadback.Load() <= 0)
+	if (GetState() == EMediaCaptureState::StopRequested && PendingFrameCount.load() <= 0)
 	{
 		// All the requested frames have been captured.
 		StopCapture(false);
 		return;
 	}
 
-	if (CapturingFrame)
+	FCaptureFrame* CapturingFrame = nullptr;
+	if(GetState() != EMediaCaptureState::StopRequested)
 	{
-		//Verify if game thread is overrunning the render thread.
-		if (CapturingFrame->bReadbackRequested)
+		CapturingFrame = FrameManager->GetNextAvailable();
+	}
+
+	if (CapturingFrame == nullptr && GetState() != EMediaCaptureState::StopRequested)
+	{
+		if (DesiredCaptureOptions.OverrunAction == EMediaCaptureOverrunAction::Flush)
 		{
-			FString TraceName = FString::Format(TEXT("MediaCapture::FlushRenderingCommands (Index {0}/{1})"), { CurrentReadbackIndex, NumberOfCaptureFrame });
-			TRACE_CPUPROFILER_EVENT_SCOPE_TEXT(*TraceName);
-			UE_LOG(LogMediaIOCore, VeryVerbose, TEXT("MediaOutput: '%s'. Flushing commands.")
-				, *MediaOutputName);
+			TRACE_CPUPROFILER_EVENT_SCOPE(MediaCapture::FlushRenderingCommands);
+			UE_LOG(LogMediaIOCore, VeryVerbose, TEXT("[%s] - Flushing commands."), *MediaOutputName);
 			FlushRenderingCommands();
+
+			// After flushing, we should have access to an available frame, if not, it's not expected.
+			CapturingFrame = FrameManager->GetNextAvailable();
+			if (CapturingFrame == nullptr)
+			{
+				UE_LOG(LogMediaIOCore, Error, TEXT("[%s] - Flushing commands didn't give us back available frames to process.") , *MediaOutputName);
+				StopCapture(false);
+				return;
+			}
+		}
+		else
+		{
+			// Selected options is to skip capturing a frame if overrun happens
+			UE_LOG(LogMediaIOCore, Verbose, TEXT("[%s] - No frames available for capture. Skipping"), *MediaOutputName);
 		}
+	}
 
+	if (CapturingFrame)
+	{
 		CapturingFrame->CaptureBaseData.SourceFrameTimecode = FApp::GetTimecode();
 		CapturingFrame->CaptureBaseData.SourceFrameTimecodeFramerate = FApp::GetTimecodeFrameRate();
 		CapturingFrame->CaptureBaseData.SourceFrameNumberRenderThread = GFrameNumber;
@@ -1327,6 +1651,8 @@ void UMediaCapture::OnEndFrame_GameThread()
 		CapturingFrame->UserData = GetCaptureFrameUserData_GameThread();
 	}
 
+	PrintFrameState();
+	
 	// Init variables for ENQUEUE_RENDER_COMMAND.
 	//The Lock only synchronize while we are copying the value to the enqueue. The viewport and the rendertarget may change while we are in the enqueue command.
 	{
@@ -1336,21 +1662,52 @@ void UMediaCapture::OnEndFrame_GameThread()
 		FSceneViewport* InCapturingSceneViewport = CapturingSceneViewportPin.Get();
 		FTextureRenderTargetResource* InTextureRenderTargetResource = CapturingRenderTarget ? CapturingRenderTarget->GameThread_GetRenderTargetResource() : nullptr;
 		FIntPoint InDesiredSize = DesiredSize;
-		FMediaCaptureStateChangedSignature InOnStateChanged = OnStateChanged;
 		UMediaCapture* InMediaCapture = this;
 
+
 		if (InCapturingSceneViewport != nullptr || InTextureRenderTargetResource != nullptr)
 		{
 			++WaitingForRenderCommandExecutionCounter;
 
 			const FRHIGPUMask SourceGPUMask = ValidSourceGPUMask;
 
+			UE_LOG(LogMediaIOCore, Verbose, TEXT("[%s] - Queuing frame to capture %d"), *InMediaCapture->MediaOutputName, CapturingFrame ? CapturingFrame->FrameId : -1);
+
 			// RenderCommand to be executed on the RenderThread
 			ENQUEUE_RENDER_COMMAND(FMediaOutputCaptureFrameCreateTexture)(
-				[InMediaCapture, CapturingFrame, ReadyFrame, InCapturingSceneViewport, InTextureRenderTargetResource, InDesiredSize, InOnStateChanged, SourceGPUMask](FRHICommandListImmediate& RHICmdList)
+				[InMediaCapture, CapturingFrame, InCapturingSceneViewport, InTextureRenderTargetResource, InDesiredSize, SourceGPUMask](FRHICommandListImmediate& RHICmdList)
 			{
 				SCOPED_GPU_MASK(RHICmdList, SourceGPUMask);
-				InMediaCapture->Capture_RenderThread(RHICmdList, InMediaCapture, CapturingFrame, ReadyFrame, InCapturingSceneViewport, InTextureRenderTargetResource, InDesiredSize, InOnStateChanged);
+
+				// Peek next pending before capturing to avoid getting the newly captured one. 
+				// It shouldn't happen but if it happens it will start waiting for the capture to complete
+				FCaptureFrame* NextPending = InMediaCapture->FrameManager->PeekNextPending();
+				if (NextPending && CapturingFrame)
+				{
+					ensure(NextPending->FrameId != CapturingFrame->FrameId);
+				}
+
+				// Capture frame
+				{
+					FRDGBuilder GraphBuilder(RHICmdList);
+
+					FTexture2DRHIRef SourceTexture = FMediaCaptureHelper::GetSourceTextureForInput(RHICmdList, InCapturingSceneViewport, InTextureRenderTargetResource);
+					InMediaCapture->ProcessCapture_RenderThread(GraphBuilder, InMediaCapture, CapturingFrame, MoveTemp(SourceTexture), InDesiredSize);
+					UE_LOG(LogMediaIOCore, Verbose, TEXT("[%s] - Captured frame %d"), *InMediaCapture->MediaOutputName, CapturingFrame->FrameId);
+					
+					GraphBuilder.Execute();
+				}
+				
+				//Process pending frame
+				UE_LOG(LogMediaIOCore, Verbose, TEXT("[%s] - Processing pending frame %d"), *InMediaCapture->MediaOutputName, NextPending ? NextPending->FrameId : -1);
+				if (NextPending)
+				{
+					InMediaCapture->ProcessReadyFrame_RenderThread(RHICmdList, InMediaCapture, NextPending);
+				}
+				
+
+				// Whatever happens, we want to decrement our counter to track enqueued commands
+				--InMediaCapture->WaitingForRenderCommandExecutionCounter;
 			});
 
 			//If auto-stopping, count the number of frame captures requested and stop when reaching 0.
@@ -1362,30 +1719,17 @@ void UMediaCapture::OnEndFrame_GameThread()
 	}
 }
 
-void UMediaCapture::Capture_RenderThread(FRHICommandListImmediate& RHICmdList,
-	UMediaCapture* InMediaCapture,
-	UE::MediaCaptureData::FCaptureFrame* CapturingFrame,
-	UE::MediaCaptureData::FCaptureFrame* ReadyFrame,
-	FSceneViewport* InCapturingSceneViewport,
-	FTextureRenderTargetResource* InTextureRenderTargetResource,
-	FIntPoint InDesiredSize,
-	FMediaCaptureStateChangedSignature InOnStateChanged)
+bool UMediaCapture::ProcessCapture_RenderThread(FRDGBuilder& GraphBuilder, UMediaCapture* InMediaCapture, UE::MediaCaptureData::FCaptureFrame* CapturingFrame, FTexture2DRHIRef InResourceToCapture,	FIntPoint InDesiredSize)
 {
 	using namespace UE::MediaCaptureData;
 	
-	TRACE_CPUPROFILER_EVENT_SCOPE(UMediaCapture::Capture_RenderThread);
-
-	// Whatever happens, we want to decrement our counter to track enqueued commands
-	ON_SCOPE_EXIT
-	{
-		--InMediaCapture->WaitingForRenderCommandExecutionCounter;
-	};
+	TRACE_CPUPROFILER_EVENT_SCOPE(UMediaCapture::ProcessCapture_RenderThread);
 
 	if (CapturingFrame)
 	{
 		// Call the capture frame algo based on the specific type of resource we are using
 		bool bHasCaptureSuceeded = false;
-		FCaptureFrameArgs CaptureArgs = { RHICmdList, InMediaCapture, InCapturingSceneViewport, InTextureRenderTargetResource, InDesiredSize };
+		FCaptureFrameArgs CaptureArgs = { GraphBuilder, InMediaCapture, MoveTemp(InResourceToCapture), InDesiredSize};
 		if (DesiredOutputResourceType == EMediaCaptureResourceType::Texture)
 		{
 			if (ensure(CapturingFrame->IsTextureResource()))
@@ -1415,8 +1759,20 @@ void UMediaCapture::Capture_RenderThread(FRHICommandListImmediate& RHICmdList,
 		{
 			InMediaCapture->SetState(EMediaCaptureState::Error);
 		}
+
+		return bHasCaptureSuceeded;
 	}
 
+	return false;
+}
+
+bool UMediaCapture::ProcessReadyFrame_RenderThread(FRHICommandListImmediate& RHICmdList, UMediaCapture* InMediaCapture, UE::MediaCaptureData::FCaptureFrame* ReadyFrame)
+{
+	using namespace UE::MediaCaptureData;
+
+	TRACE_CPUPROFILER_EVENT_SCOPE(UMediaCapture::ProcessReadyFrame_RenderThread);
+
+	bool bWasFrameProcessed = true;
 	if (InMediaCapture->GetState() != EMediaCaptureState::Error)
 	{
 		if (ReadyFrame->bReadbackRequested)
@@ -1437,27 +1793,40 @@ void UMediaCapture::Capture_RenderThread(FRHICommandListImmediate& RHICmdList,
 			void* ColorDataBuffer = nullptr;
 			int32 RowStride = 0;
 
-			// Texture readback does no verification if it's ready. It is assumed it will stall if it's not
-			// Buffer readback will verify and block GPU until idle in case it's not ready otherwise it ensures
+			// If readback is ready, proceed. 
+			// If not, proceed with locking only if we are in flush mode since it will block until gpu is idle.
+			const bool bIsReadbackReady = InMediaCapture->DesiredCaptureOptions.OverrunAction == EMediaCaptureOverrunAction::Flush || ReadyFrame->IsReadbackReady(GPUMask) == true;
+			if (bIsReadbackReady)
 			{
 				SCOPE_CYCLE_COUNTER(STAT_MediaCapture_RenderThread_LockResource);
 				ColorDataBuffer = ReadyFrame->Lock(RHICmdList, GPUMask, RowStride);
 			}
+			else
+			{
+				UE_LOG(LogMediaIOCore, VeryVerbose, TEXT("[%s] - Readback %d not ready. Skipping."), *InMediaCapture->MediaOutputName, ReadyFrame->FrameId);
+				bWasFrameProcessed = false;
+			}
 
+			if(ColorDataBuffer)
 			{
-				SCOPE_CYCLE_COUNTER(STAT_MediaCapture_RenderThread_CaptureCallback)
+				{
+					SCOPE_CYCLE_COUNTER(STAT_MediaCapture_RenderThread_CaptureCallback)
 
-				// The Width/Height of the surface may be different then the DesiredOutputSize : Some underlying implementations enforce a specific stride, therefore
-				// there may be padding at the end of each row.
-				InMediaCapture->OnFrameCaptured_RenderingThread(ReadyFrame->CaptureBaseData, ReadyFrame->UserData, ColorDataBuffer, InMediaCapture->DesiredOutputSize.X, InMediaCapture->DesiredOutputSize.Y, RowStride);
-			}
-			ReadyFrame->bReadbackRequested = false;
+					// The Width/Height of the surface may be different then the DesiredOutputSize : Some underlying implementations enforce a specific stride, therefore
+					// there may be padding at the end of each row.
+					InMediaCapture->OnFrameCaptured_RenderingThread(ReadyFrame->CaptureBaseData, ReadyFrame->UserData, ColorDataBuffer, InMediaCapture->DesiredOutputSize.X, InMediaCapture->DesiredOutputSize.Y, RowStride);
+				}
 
-			const int32 FrameIndex = CaptureFrames.IndexOfByPredicate([&CapturingFrame](const TUniquePtr<FCaptureFrame>& InCaptureFrame) { return InCaptureFrame.Get() == CapturingFrame; });
-			UE_LOG(LogMediaIOCore, VeryVerbose, TEXT("Copy completed for capture frame %d."), FrameIndex);
-			--NumberOfResourcesToReadback;
+				ReadyFrame->Unlock();
+			}
 
-			ReadyFrame->Unlock();
+			if (bIsReadbackReady)
+			{
+				UE_LOG(LogMediaIOCore, Verbose, TEXT("[%s] - Completed pending frame %d."), *InMediaCapture->MediaOutputName, ReadyFrame->FrameId);
+				ReadyFrame->bReadbackRequested = false;
+				--PendingFrameCount;
+				InMediaCapture->FrameManager->CompleteNextPending(*ReadyFrame);
+			}
 		}
 		else if (InMediaCapture->bShouldCaptureRHIResource && ReadyFrame->bDoingGPUCopy)
 		{
@@ -1471,12 +1840,31 @@ void UMediaCapture::Capture_RenderThread(FRHICommandListImmediate& RHICmdList,
 
 				TRACE_CPUPROFILER_EVENT_SCOPE(UMediaCapture::RHIResourceCaptured);
 				SCOPE_CYCLE_COUNTER(STAT_MediaCapture_RenderThread_CaptureCallback)
-					InMediaCapture->OnRHIResourceCaptured_RenderingThread(ReadyFrame->CaptureBaseData, ReadyFrame->UserData, TextureFrame->GetRHIResource());
+				InMediaCapture->OnRHIResourceCaptured_RenderingThread(ReadyFrame->CaptureBaseData, ReadyFrame->UserData, TextureFrame->GetRHIResource());
 
+				UE_LOG(LogMediaIOCore, Verbose, TEXT("[%s] - Completed pending frame %d."), *InMediaCapture->MediaOutputName, ReadyFrame->FrameId);
 				ReadyFrame->bDoingGPUCopy = false;
+				InMediaCapture->FrameManager->CompleteNextPending(*ReadyFrame);
+				--PendingFrameCount;
 			}
 		}
 	}
+
+	return bWasFrameProcessed;
+}
+
+void UMediaCapture::PrintFrameState()
+{
+	using namespace UE::MediaCaptureData;
+
+	TStringBuilder<256> ReadbackInfoBuilder;
+	ReadbackInfoBuilder << "\n";
+	TConstArrayView<TUniquePtr<FCaptureFrame>> Frames = FrameManager->GetFrames();
+	for (int32 Index = 0; Index < Frames.Num(); Index++)
+	{
+		ReadbackInfoBuilder << FString::Format(TEXT("Frame {0} readback requested: {1}\n"), { Index, Frames[Index]->IsPending() ? TEXT("true") : TEXT("false") });
+	}
+	UE_LOG(LogMediaIOCore, VeryVerbose, TEXT("%s"), ReadbackInfoBuilder.GetData());
 }
 
 /* namespace MediaCaptureDetails implementation
diff --git a/Engine/Plugins/Media/MediaIOFramework/Source/MediaIOCore/Public/MediaCapture.h b/Engine/Plugins/Media/MediaIOFramework/Source/MediaIOCore/Public/MediaCapture.h
--- a/Engine/Plugins/Media/MediaIOFramework/Source/MediaIOCore/Public/MediaCapture.h
+++ b/Engine/Plugins/Media/MediaIOFramework/Source/MediaIOCore/Public/MediaCapture.h
@@ -6,6 +6,8 @@
 #include "UObject/ObjectMacros.h"
 #include "UObject/Object.h"
 
+#include <atomic>
+#include "Containers/SpscQueue.h"
 #include "HAL/CriticalSection.h"
 #include "MediaOutput.h"
 #include "Misc/Timecode.h"
@@ -14,7 +16,6 @@
 #include "RenderGraphResources.h"
 #include "RHI.h"
 #include "RHIResources.h"
-#include "Templates/Atomic.h"
 
 #include "MediaCapture.generated.h"
 
@@ -24,9 +25,11 @@ class UTextureRenderTarget2D;
 class FRHICommandListImmediate;
 class FRHIGPUTextureReadback;
 class FRHIGPUBufferReadback;
+class FRDGBuilder;
 
 namespace UE::MediaCaptureData
 {
+	class FFrameManager;
 	class FCaptureFrame;
 	class FMediaCaptureHelper;
 }
@@ -90,6 +93,26 @@ enum class EMediaCaptureCroppingType : uint8
 	Custom,
 };
 
+/**
+ * Action when overrun occurs
+ */
+UENUM(BlueprintType)
+enum class EMediaCaptureOverrunAction : uint8
+{
+	/** Flush rendering thread such that all scheduled commands are executed. */
+	Flush,
+	/** Skip capturing a frame if readback is trailing too much. */
+	Skip,
+};
+
+/**
+ * Description of resource to be captured when in rhi capture / immediate mode.
+ */
+struct MEDIAIOCORE_API FRHICaptureResourceDescription
+{
+	FIntPoint ResourceSize = FIntPoint::ZeroValue;
+};
+
 /**
  * Base class of additional data that can be stored for each requested capture.
  */
@@ -101,6 +124,11 @@ struct MEDIAIOCORE_API FMediaCaptureOptions
 	FMediaCaptureOptions();
 
 public:
+
+	/** Action to do when game thread overruns render thread and all frames are in flights being captured / readback. */
+	UPROPERTY(BlueprintReadWrite, EditAnywhere, Category = "MediaCapture")
+	EMediaCaptureOverrunAction OverrunAction = EMediaCaptureOverrunAction::Flush;
+
 	/** Crop the captured SceneViewport or TextureRenderTarget2D to the desired size. */
 	UPROPERTY(BlueprintReadWrite, EditAnywhere, Category="MediaCapture")
 	EMediaCaptureCroppingType Crop;
@@ -117,6 +145,7 @@ public:
 	 * @note Only valid when a size is specified by the MediaOutput.
 	 * @note For viewport, the window size will not change. Only the viewport will be resized.
 	 * @note For RenderTarget, the asset will be modified and resized to the desired size.
+	 * @note Not valid for immediate capture of RHI resource
 	 */
 	UPROPERTY(BlueprintReadWrite, EditAnywhere, Category="MediaCapture")
 	bool bResizeSourceBuffer;
@@ -152,7 +181,6 @@ public:
 	UPROPERTY(BlueprintReadWrite, EditAnywhere, Category = "MediaCapture")
 	bool bAutostopOnCapture;
 
-
 	/** The number of images to capture*/
 	UPROPERTY(BlueprintReadWrite, EditAnywhere, Category = "MediaCapture", meta=(editcondition="bAutostopOnCapture"))
 	int32 NumberOfFramesToCapture;
@@ -191,6 +219,14 @@ public:
 	 */
 	bool CaptureSceneViewport(TSharedPtr<FSceneViewport>& SceneViewport, FMediaCaptureOptions CaptureOptions);
 
+	/**
+	 * Stop the actual capture if there is one.
+	 * Then start a capture for a RHI resource.
+	 * @note For this to work, it is expected that CaptureImmediate_RenderThread is called periodically with the resource to capture.
+	 * @return True if the capture was successfully initialize
+	 */
+	bool CaptureRHITexture(const FRHICaptureResourceDescription& ResourceDescription, FMediaCaptureOptions CaptureOptions);
+
 	/**
 	 * Stop the current capture if there is one.
 	 * Then find and capture every frame from active SceneViewport.
@@ -220,7 +256,6 @@ public:
 	 */
 	bool UpdateSceneViewport(TSharedPtr<FSceneViewport>& SceneViewport);
 
-
 	/**
 	 * Update the current capture with every frame for a TextureRenderTarget2D.
 	 * The TextureRenderTarget2D needs to be of the same size and have the same pixel format as requested by the media output.
@@ -228,7 +263,9 @@ public:
 	 */
 	UFUNCTION(BlueprintCallable, Category = "Media|Output")
 	bool UpdateTextureRenderTarget2D(UTextureRenderTarget2D* RenderTarget);
-
+	
+	/** Captures a resource immediately from the render thread. Used in RHI_RESOURCE capture mode */
+	void CaptureImmediate_RenderThread(FRDGBuilder& GraphBuilder, FRHITexture* InSourceTexture);
 
 	/**
 	 * Stop the previous requested capture.
@@ -277,8 +314,25 @@ public:
 protected:
 	//~ UMediaCapture interface
 	virtual bool ValidateMediaOutput() const;
-	virtual bool CaptureSceneViewportImpl(TSharedPtr<FSceneViewport>& InSceneViewport) { return true; }
-	virtual bool CaptureRenderTargetImpl(UTextureRenderTarget2D* InRenderTarget) { return true; }
+
+	UE_DEPRECATED(5.1, "This method has been deprecated. Please override InitializeCapture() and PostInitializeCapture variants if needed instead.")
+	virtual bool CaptureSceneViewportImpl(TSharedPtr<FSceneViewport>& InSceneViewport) final;
+
+	UE_DEPRECATED(5.1, "This method has been deprecated. Please override InitializeCapture() and PostInitializeCapture variants if needed instead.")
+	virtual bool CaptureRenderTargetImpl(UTextureRenderTarget2D* InRenderTarget) final;
+
+	/** Initialization method to prepare implementation for capture */
+	virtual bool InitializeCapture() PURE_VIRTUAL(UMediaCapture::InitializeCapture, return false; );
+	
+	/** Called after initialize for viewport capture type */
+	virtual bool PostInitializeCaptureViewport(TSharedPtr<FSceneViewport>& InSceneViewport) { return true; }
+
+	/** Called after initialize for render target capture type */
+	virtual bool PostInitializeCaptureRenderTarget(UTextureRenderTarget2D* InRenderTarget) { return true; }
+
+	/** Called after initialize for rhi resource capture type */
+	virtual bool PostInitializeCaptureRHIResource(const FRHICaptureResourceDescription& InResourceDescription) { return true; }
+
 	virtual bool UpdateSceneViewportImpl(TSharedPtr<FSceneViewport>& InSceneViewport) { return true; }
 	virtual bool UpdateRenderTargetImpl(UTextureRenderTarget2D* InRenderTarget) { return true; }
 	virtual void StopCaptureImpl(bool bAllowPendingFrameToBeProcess) { }
@@ -387,15 +441,20 @@ protected:
 	}
 
 
+
 protected:
 	UTextureRenderTarget2D* GetTextureRenderTarget() { return CapturingRenderTarget; }
 	TSharedPtr<FSceneViewport> GetCapturingSceneViewport() { return CapturingSceneViewport.Pin(); }
 	EMediaCaptureConversionOperation GetConversionOperation() const { return ConversionOperation; }
+	const FString& GetCaptureSourceType() const;
 	void SetState(EMediaCaptureState InNewState);
 
 
 private:
 	void InitializeOutputResources(int32 InNumberOfBuffers);
+	
+	/** Used by RHI capture to buffer captured frame data dependant on game thread such as timecode */
+	void OnBeginFrame_GameThread();
 	void OnEndFrame_GameThread();
 	void CacheMediaOutput(EMediaCaptureSourceType InSourceType);
 	void CacheOutputOptions();
@@ -407,9 +466,10 @@ private:
 	void SetFixedViewportSize(TSharedPtr<FSceneViewport> InSceneViewport);
 	void ResetFixedViewportSize(TSharedPtr<FSceneViewport> InViewport, bool bInFlushRenderingCommands);
 
-	void Capture_RenderThread(FRHICommandListImmediate& RHICmdList, UMediaCapture* InMediaCapture, UE::MediaCaptureData::FCaptureFrame* CapturingFrame, UE::MediaCaptureData::FCaptureFrame* ReadyFrame,
-		FSceneViewport* InCapturingSceneViewport, FTextureRenderTargetResource* InTextureRenderTargetResource,
-		FIntPoint InDesiredSize, FMediaCaptureStateChangedSignature InOnStateChanged);
+	bool ProcessCapture_RenderThread(FRDGBuilder& GraphBuilder, UMediaCapture* InMediaCapture, UE::MediaCaptureData::FCaptureFrame* CapturingFrame, FTexture2DRHIRef InResourceToCapture, FIntPoint InDesiredSize);
+	bool ProcessReadyFrame_RenderThread(FRHICommandListImmediate& RHICmdList, UMediaCapture* InMediaCapture, UE::MediaCaptureData::FCaptureFrame* ReadyFrame);
+
+	void PrintFrameState();
 
 protected:
 
@@ -423,13 +483,15 @@ protected:
 	/** Valid source gpu mask of the source. */
 	std::atomic<FRHIGPUMask> ValidSourceGPUMask;
 
+	/** Type of capture currently being done. i.e Viewport, RT, RHI Resource */
+	EMediaCaptureSourceType CaptureSourceType = EMediaCaptureSourceType::SCENE_VIEWPORT;
+
 private:
 
 	UPROPERTY(Transient)
 	TObjectPtr<UTextureRenderTarget2D> CapturingRenderTarget;
 	
-	TArray<TUniquePtr<UE::MediaCaptureData::FCaptureFrame>> CaptureFrames;
-	int32 CurrentReadbackIndex = 0;
+	TUniquePtr<UE::MediaCaptureData::FFrameManager> FrameManager;
 	int32 NumberOfCaptureFrame = 2;
 	int32 CaptureRequestCount = 0;
 	EMediaCaptureState MediaState = EMediaCaptureState::Stopped;
@@ -448,8 +510,15 @@ private:
 	bool bUseRequestedTargetSize = false;
 
 	bool bViewportHasFixedViewportSize = false;
-	TAtomic<bool> bOutputResourcesInitialized;
-	TAtomic<bool> bShouldCaptureRHIResource;
-	TAtomic<int32> WaitingForRenderCommandExecutionCounter;
-	TAtomic<int32> NumberOfResourcesToReadback;
+	std::atomic<bool> bOutputResourcesInitialized;
+	std::atomic<bool> bShouldCaptureRHIResource;
+	std::atomic<int32> WaitingForRenderCommandExecutionCounter;
+	std::atomic<int32> PendingFrameCount;
+
+	struct FQueuedCaptureData
+	{
+		FCaptureBaseData BaseData;
+		TSharedPtr<FMediaCaptureUserData> UserData;
+	};
+	TSpscQueue<FQueuedCaptureData> CaptureDataQueue;
 };
diff --git a/Engine/Plugins/Media/MediaIOFramework/Source/MediaIOCore/Public/MediaOutput.h b/Engine/Plugins/Media/MediaIOFramework/Source/MediaIOCore/Public/MediaOutput.h
--- a/Engine/Plugins/Media/MediaIOFramework/Source/MediaIOCore/Public/MediaOutput.h
+++ b/Engine/Plugins/Media/MediaIOFramework/Source/MediaIOCore/Public/MediaOutput.h
@@ -27,6 +27,7 @@ enum class EMediaCaptureSourceType : int8
 {
 	RENDER_TARGET,
 	SCENE_VIEWPORT,
+	RHI_RESOURCE,
 };
 
 /**
diff --git a/Engine/Plugins/VirtualProduction/Rivermax/RivermaxMedia/Source/RivermaxMedia/Private/RivermaxMediaCapture.cpp b/Engine/Plugins/VirtualProduction/Rivermax/RivermaxMedia/Source/RivermaxMedia/Private/RivermaxMediaCapture.cpp
--- a/Engine/Plugins/VirtualProduction/Rivermax/RivermaxMedia/Source/RivermaxMedia/Private/RivermaxMediaCapture.cpp
+++ b/Engine/Plugins/VirtualProduction/Rivermax/RivermaxMedia/Source/RivermaxMedia/Private/RivermaxMediaCapture.cpp
@@ -111,27 +111,14 @@ bool URivermaxMediaCapture::ValidateMediaOutput() const
 	return true;
 }
 
-bool URivermaxMediaCapture::CaptureSceneViewportImpl(TSharedPtr<FSceneViewport>& InSceneViewport)
-{
-	URivermaxMediaOutput* RivermaxOutput = CastChecked<URivermaxMediaOutput>(MediaOutput);
-	const bool bResult = Initialize(RivermaxOutput);
-	if (bResult)
-	{
-#if WITH_EDITOR
-		RivermaxMediaCaptureAnalytics::SendCaptureEvent(RivermaxOutput->GetRequestedSize(), RivermaxOutput->FrameRate, TEXT("SceneViewport"));
-#endif
-	}
-	return bResult;
-}
-
-bool URivermaxMediaCapture::CaptureRenderTargetImpl(UTextureRenderTarget2D* InRenderTarget)
+bool URivermaxMediaCapture::InitializeCapture()
 {
 	URivermaxMediaOutput* RivermaxOutput = CastChecked<URivermaxMediaOutput>(MediaOutput);
 	bool bResult = Initialize(RivermaxOutput);
 #if WITH_EDITOR
 	if (bResult)
 	{
-		RivermaxMediaCaptureAnalytics::SendCaptureEvent(RivermaxOutput->GetRequestedSize(), RivermaxOutput->FrameRate, TEXT("RenderTarget"));
+		RivermaxMediaCaptureAnalytics::SendCaptureEvent(GetDesiredSize(), RivermaxOutput->FrameRate, GetCaptureSourceType());
 	}
 #endif
 	return bResult;
diff --git a/Engine/Plugins/VirtualProduction/Rivermax/RivermaxMedia/Source/RivermaxMedia/Public/RivermaxMediaCapture.h b/Engine/Plugins/VirtualProduction/Rivermax/RivermaxMedia/Source/RivermaxMedia/Public/RivermaxMediaCapture.h
--- a/Engine/Plugins/VirtualProduction/Rivermax/RivermaxMedia/Source/RivermaxMedia/Public/RivermaxMediaCapture.h
+++ b/Engine/Plugins/VirtualProduction/Rivermax/RivermaxMedia/Source/RivermaxMedia/Public/RivermaxMediaCapture.h
@@ -32,8 +32,7 @@ public:
 	virtual bool HasFinishedProcessing() const override;
 protected:
 	virtual bool ValidateMediaOutput() const override;
-	virtual bool CaptureSceneViewportImpl(TSharedPtr<FSceneViewport>& InSceneViewport) override;
-	virtual bool CaptureRenderTargetImpl(UTextureRenderTarget2D* InRenderTarget) override;
+	virtual bool InitializeCapture() override;
 	virtual bool UpdateSceneViewportImpl(TSharedPtr<FSceneViewport>& InSceneViewport) override;
 	virtual bool UpdateRenderTargetImpl(UTextureRenderTarget2D* InRenderTarget) override;
 	virtual void StopCaptureImpl(bool bAllowPendingFrameToBeProcess) override;
