- Fixing crash with media capture frame manager

diff --git a/Engine/Plugins/Media/MediaIOFramework/Source/MediaIOCore/Private/MediaCapture.cpp b/Engine/Plugins/Media/MediaIOFramework/Source/MediaIOCore/Private/MediaCapture.cpp
--- a/Engine/Plugins/Media/MediaIOFramework/Source/MediaIOCore/Private/MediaCapture.cpp
+++ b/Engine/Plugins/Media/MediaIOFramework/Source/MediaIOCore/Private/MediaCapture.cpp
@@ -333,6 +333,15 @@ namespace UE::MediaCaptureData
 	{
 	public:
 
+		~FFrameManager()
+		{
+			ENQUEUE_RENDER_COMMAND(MediaCaptureFrameManagerCleaning)(
+				[FramesToBeReleased = MoveTemp(CaptureFrames)](FRHICommandListImmediate& RHICmdList) mutable
+			{
+				FramesToBeReleased.Reset();
+			});
+		}
+
 		TConstArrayView<TUniquePtr<FCaptureFrame>> GetFrames() const
 		{
 			return CaptureFrames;
@@ -740,7 +749,6 @@ FMediaCaptureOptions::FMediaCaptureOptions()
 UMediaCapture::UMediaCapture(const FObjectInitializer& ObjectInitializer)
 	: Super(ObjectInitializer)
 	, ValidSourceGPUMask(FRHIGPUMask::All())
-	, FrameManager(MakeUnique<UE::MediaCaptureData::FFrameManager>())
 	, bOutputResourcesInitialized(false)
 	, bShouldCaptureRHIResource(false)
 	, WaitingForRenderCommandExecutionCounter(0)
@@ -1198,13 +1206,13 @@ void UMediaCapture::StopCapture(bool bAllowPendingFrameToBeProcess)
 			//Do not flush when auto stopping to avoid hitches.
 			if(DesiredCaptureOptions.bAutostopOnCapture != true)
 			{
-			while (WaitingForRenderCommandExecutionCounter.load() > 0)
-			{
-				FlushRenderingCommands();
+				while (WaitingForRenderCommandExecutionCounter.load() > 0)
+				{
+					FlushRenderingCommands();
+				}
 			}
 		}
 	}
-	}
 	else
 	{
 		if (GetState() != EMediaCaptureState::Stopped)
@@ -1230,14 +1238,6 @@ void UMediaCapture::StopCapture(bool bAllowPendingFrameToBeProcess)
 			DesiredCaptureOptions = FMediaCaptureOptions();
 			ConversionOperation = EMediaCaptureConversionOperation::NONE;
 			MediaOutputName.Reset();
-
-			// CaptureFrames contains RHI resources, therefore should be released on Render Tread thread.
-			// Keep references frames to be released in a temporary array and clear CaptureFrames on Game Thread.
-			ENQUEUE_RENDER_COMMAND(MediaOutputReleaseCaptureFrames)(
-				[FrameManagerToRelease = MoveTemp(FrameManager)](FRHICommandListImmediate& RHICmdList) mutable
-				{
-					FrameManagerToRelease.Reset();
-				});
 		}
 	}
 }
@@ -1466,11 +1466,13 @@ void UMediaCapture::InitializeOutputResources(int32 InNumberOfBuffers)
 		return;
 	}
 
+	// Recreate frame manager which can trigger cleaning up its captured frames if it exists
+	FrameManager = MakeUnique<UE::MediaCaptureData::FFrameManager>();
+
 	bOutputResourcesInitialized = false;
 	NumberOfCaptureFrame = InNumberOfBuffers;
 
 	check(NumberOfCaptureFrame >= 2);
-	check(FrameManager->GetFrames().Num() == 0);
 
 	UMediaCapture* This = this;
 	ENQUEUE_RENDER_COMMAND(MediaOutputCaptureFrameCreateResources)(
