- Fixinng unacceptable words in mediaio

diff --git a/Engine/Plugins/Media/MediaIOFramework/Source/MediaIOCore/Private/MediaIOCoreAudioOutput.cpp b/Engine/Plugins/Media/MediaIOFramework/Source/MediaIOCore/Private/MediaIOCoreAudioOutput.cpp
--- a/Engine/Plugins/Media/MediaIOFramework/Source/MediaIOCore/Private/MediaIOCoreAudioOutput.cpp
+++ b/Engine/Plugins/Media/MediaIOFramework/Source/MediaIOCore/Private/MediaIOCoreAudioOutput.cpp
@@ -75,7 +75,7 @@ FMediaIOAudioCapture::~FMediaIOAudioCapture()
 void FMediaIOAudioCapture::OnNewSubmixBuffer(const USoundSubmix* InOwningSubmix, float* InAudioData, int32 InNumSamples, int32 InNumChannels, const int32 InSampleRate, double InAudioClock)
 {
 	check(InOwningSubmix);
-	if (InOwningSubmix->GetFName() == MasterSubmixName)
+	if (InOwningSubmix->GetFName() == PrimarySubmixName)
 	{
 		if (ensureMsgf(NumChannels == InNumChannels, TEXT("Expected %d channels from submix buffer but got %d instead."), NumChannels, InNumChannels))
 		{
@@ -173,7 +173,7 @@ void FMediaIOAudioCapture::RegisterBufferListener(FAudioDevice* AudioDevice)
 		Audio::FMixerDevice* MixerDevice = static_cast<Audio::FMixerDevice*>(AudioDevice);
 		NumChannels = MixerDevice->GetDeviceOutputChannels();
 		SampleRate = MixerDevice->GetSampleRate();
-		MasterSubmixName = *GetDefault<UAudioSettings>()->MasterSubmix.GetAssetName();
+		PrimarySubmixName = *GetDefault<UAudioSettings>()->MasterSubmix.GetAssetName();
 		AudioDevice->RegisterSubmixBufferListener(this);
 	}
 }
diff --git a/Engine/Plugins/Media/MediaIOFramework/Source/MediaIOCore/Public/MediaIOCoreAudioOutput.h b/Engine/Plugins/Media/MediaIOFramework/Source/MediaIOCore/Public/MediaIOCoreAudioOutput.h
--- a/Engine/Plugins/Media/MediaIOFramework/Source/MediaIOCore/Public/MediaIOCoreAudioOutput.h
+++ b/Engine/Plugins/Media/MediaIOFramework/Source/MediaIOCore/Public/MediaIOCoreAudioOutput.h
@@ -141,8 +141,8 @@ private:
 	/** Sample rate on the engine side. */ 
 	uint32 SampleRate = 0;
 
-	// Used to make sure we only accumulate audio from the master submix. 
-	FName MasterSubmixName;
+	// Used to make sure we only accumulate audio from the primary submix. 
+	FName PrimarySubmixName;
 
 	/** Number of channels on the engine side. */
 	int32 NumChannels = 0;
