Fix wrong/default format showed when autodetected
Default Aja/BM Sources to first device, autodetected
Remove old auto detect button from aja media source
Change "Enforce Format" to "Auto" checkbox

diff --git a/Engine/Plugins/Media/AjaMedia/Source/AjaMedia/Private/Assets/AjaMediaSource.cpp b/Engine/Plugins/Media/AjaMedia/Source/AjaMedia/Private/Assets/AjaMediaSource.cpp
--- a/Engine/Plugins/Media/AjaMedia/Source/AjaMedia/Private/Assets/AjaMediaSource.cpp
+++ b/Engine/Plugins/Media/AjaMedia/Source/AjaMedia/Private/Assets/AjaMediaSource.cpp
@@ -3,10 +3,9 @@
 #include "AjaMediaSource.h"
 
 #include "Aja.h"
+#include "AjaDeviceProvider.h"
 #include "AjaMediaPrivate.h"
-
 #include "MediaIOCorePlayerBase.h"
-#include "UObject/EnterpriseObjectVersion.h"
 
 UAjaMediaSource::UAjaMediaSource()
 	: bCaptureWithAutoCirculating(true)
@@ -23,6 +22,12 @@ UAjaMediaSource::UAjaMediaSource()
 	, bEncodeTimecodeInTexel(false)
 {
 	MediaConfiguration.bIsInput = true;
+	const FAjaDeviceProvider DeviceProvider;
+	const TArray<FMediaIOConfiguration> Configurations = DeviceProvider.GetConfigurations();
+	if (Configurations.Num())
+	{
+		MediaConfiguration = Configurations[0];
+	}
 }
 
 /*
diff --git a/Engine/Plugins/Media/AjaMedia/Source/AjaMediaEditor/Private/AjaMediaEditorModule.cpp b/Engine/Plugins/Media/AjaMedia/Source/AjaMediaEditor/Private/AjaMediaEditorModule.cpp
--- a/Engine/Plugins/Media/AjaMedia/Source/AjaMediaEditor/Private/AjaMediaEditorModule.cpp
+++ b/Engine/Plugins/Media/AjaMedia/Source/AjaMediaEditor/Private/AjaMediaEditorModule.cpp
@@ -7,7 +7,6 @@
 #include "AjaTimecodeProvider.h"
 
 #include "Customizations/AjaMediaTimecodeReferenceCustomization.h"
-#include "Customizations/AjaMediaSourceDetailsCustomization.h"
 
 #include "Brushes/SlateImageBrush.h"
 #include "Interfaces/IPluginManager.h"
@@ -56,7 +55,6 @@ private:
 	{
 		FPropertyEditorModule& PropertyModule = FModuleManager::LoadModuleChecked<FPropertyEditorModule>("PropertyEditor");
 		PropertyModule.RegisterCustomPropertyTypeLayout(FAjaMediaTimecodeReference::StaticStruct()->GetFName(), FOnGetPropertyTypeCustomizationInstance::CreateStatic(&FAjaMediaTimecodeReferenceCustomization::MakeInstance));
-		PropertyModule.RegisterCustomClassLayout(UAjaMediaSource::StaticClass()->GetFName(), FOnGetDetailCustomizationInstance::CreateStatic(&FAjaMediaSourceDetailsCustomization::MakeInstance));
 		
 	}
 
diff --git a/Engine/Plugins/Media/AjaMedia/Source/AjaMediaEditor/Private/Customizations/AjaMediaSourceDetailsCustomization.cpp b/Engine/Plugins/Media/AjaMedia/Source/AjaMediaEditor/Private/Customizations/AjaMediaSourceDetailsCustomization.cpp
deleted file mode 100644
diff --git a/Engine/Plugins/Media/AjaMedia/Source/AjaMediaEditor/Private/Customizations/AjaMediaSourceDetailsCustomization.h b/Engine/Plugins/Media/AjaMedia/Source/AjaMediaEditor/Private/Customizations/AjaMediaSourceDetailsCustomization.h
deleted file mode 100644
diff --git a/Engine/Plugins/Media/BlackmagicMedia/Source/BlackmagicMedia/Private/Assets/BlackmagicMediaSource.cpp b/Engine/Plugins/Media/BlackmagicMedia/Source/BlackmagicMedia/Private/Assets/BlackmagicMediaSource.cpp
--- a/Engine/Plugins/Media/BlackmagicMedia/Source/BlackmagicMedia/Private/Assets/BlackmagicMediaSource.cpp
+++ b/Engine/Plugins/Media/BlackmagicMedia/Source/BlackmagicMedia/Private/Assets/BlackmagicMediaSource.cpp
@@ -3,6 +3,7 @@
 #include "BlackmagicMediaSource.h"
 
 #include "Blackmagic.h"
+#include "BlackmagicDeviceProvider.h"
 #include "BlackmagicMediaPrivate.h"
 #include "IBlackmagicMediaModule.h"
 
@@ -20,6 +21,15 @@ UBlackmagicMediaSource::UBlackmagicMediaSource()
 	, bEncodeTimecodeInTexel(false)
 {
 	MediaConfiguration.bIsInput = true;
+	if (!MediaConfiguration.IsValid())
+	{
+		const FBlackmagicDeviceProvider DeviceProvider;
+		const TArray<FMediaIOConfiguration> Configurations = DeviceProvider.GetConfigurations();
+		if (Configurations.Num())
+		{
+			MediaConfiguration = Configurations[0];
+		}
+	}
 }
 
 /*
diff --git a/Engine/Plugins/Media/BlackmagicMedia/Source/BlackmagicMedia/Private/Blackmagic/BlackmagicDeviceProvider.cpp b/Engine/Plugins/Media/BlackmagicMedia/Source/BlackmagicMedia/Private/Blackmagic/BlackmagicDeviceProvider.cpp
--- a/Engine/Plugins/Media/BlackmagicMedia/Source/BlackmagicMedia/Private/Blackmagic/BlackmagicDeviceProvider.cpp
+++ b/Engine/Plugins/Media/BlackmagicMedia/Source/BlackmagicMedia/Private/Blackmagic/BlackmagicDeviceProvider.cpp
@@ -445,8 +445,17 @@ UMediaSource* FBlackmagicDeviceProvider::CreateMediaSource(
 }
 
 
-FText FBlackmagicDeviceProvider::ToText(const FMediaIOConfiguration& InConfiguration) const
+FText FBlackmagicDeviceProvider::ToText(const FMediaIOConfiguration& InConfiguration, bool bInIsAutoDetected) const
 {
+	if (bInIsAutoDetected)
+	{
+		return FText::Format(LOCTEXT("FMediaIOConfigurationToText", "[{0}] - {1} [device{2}/auto]")
+		, InConfiguration.bIsInput ? LOCTEXT("In", "In") : LOCTEXT("Out", "Out")
+		, FText::FromName(InConfiguration.MediaConnection.Device.DeviceName)
+		, FText::AsNumber(InConfiguration.MediaConnection.Device.DeviceIdentifier)
+		);
+	}
+	
 	if (InConfiguration.IsValid())
 	{
 		return FText::Format(LOCTEXT("FMediaIOConfigurationToText", "[{0}] - {1} [device{2}/{3}]")
diff --git a/Engine/Plugins/Media/BlackmagicMedia/Source/BlackmagicMedia/Public/BlackmagicDeviceProvider.h b/Engine/Plugins/Media/BlackmagicMedia/Source/BlackmagicMedia/Public/BlackmagicDeviceProvider.h
--- a/Engine/Plugins/Media/BlackmagicMedia/Source/BlackmagicMedia/Public/BlackmagicDeviceProvider.h
+++ b/Engine/Plugins/Media/BlackmagicMedia/Source/BlackmagicMedia/Public/BlackmagicDeviceProvider.h
@@ -49,7 +49,7 @@ public:
 	virtual UMediaSource* CreateMediaSource(const FMediaIOConfiguration& InConfiguration,
 		UObject* Outer) const override;
 
-	virtual FText ToText(const FMediaIOConfiguration& InConfiguration) const override;
+	virtual FText ToText(const FMediaIOConfiguration& InConfiguration, bool bInIsAutoDetected = false) const override;
 	virtual FText ToText(const FMediaIOConnection& InConnection) const override;
 	virtual FText ToText(const FMediaIOOutputConfiguration& InConfiguration) const override;
 
diff --git a/Engine/Plugins/Media/MediaIOFramework/Source/MediaIOCore/Private/MediaIOCoreDeviceProvider.cpp b/Engine/Plugins/Media/MediaIOFramework/Source/MediaIOCore/Private/MediaIOCoreDeviceProvider.cpp
--- a/Engine/Plugins/Media/MediaIOFramework/Source/MediaIOCore/Private/MediaIOCoreDeviceProvider.cpp
+++ b/Engine/Plugins/Media/MediaIOFramework/Source/MediaIOCore/Private/MediaIOCoreDeviceProvider.cpp
@@ -29,8 +29,16 @@ FText IMediaIOCoreDeviceProvider::GetTransportName(EMediaIOTransportType InLinkT
 	return Transport;
 }
 
-FText IMediaIOCoreDeviceProvider::ToText(const FMediaIOConfiguration& InConfiguration) const
+FText IMediaIOCoreDeviceProvider::ToText(const FMediaIOConfiguration& InConfiguration, bool bInIsAutoDetected) const
 {
+	if (bInIsAutoDetected)
+	{
+		return FText::Format(LOCTEXT("FMediaIOConfigurationToText", "{0} - {1} [device{2}/auto]")
+				, InConfiguration.bIsInput ? LOCTEXT("In", "In") : LOCTEXT("Out", "Out")
+				, FText::FromName(InConfiguration.MediaConnection.Device.DeviceName)
+				, FText::AsNumber(InConfiguration.MediaConnection.Device.DeviceIdentifier)
+				);
+	}
 	if (InConfiguration.IsValid())
 	{
 		return FText::Format(LOCTEXT("FMediaIOConfigurationToText", "{0} - {1} [device{2}/{3}{4}/{5}]")
diff --git a/Engine/Plugins/Media/MediaIOFramework/Source/MediaIOCore/Public/IMediaIOCoreDeviceProvider.h b/Engine/Plugins/Media/MediaIOFramework/Source/MediaIOCore/Public/IMediaIOCoreDeviceProvider.h
--- a/Engine/Plugins/Media/MediaIOFramework/Source/MediaIOCore/Public/IMediaIOCoreDeviceProvider.h
+++ b/Engine/Plugins/Media/MediaIOFramework/Source/MediaIOCore/Public/IMediaIOCoreDeviceProvider.h
@@ -30,7 +30,7 @@ public:
 	virtual UMediaSource* CreateMediaSource(const FMediaIOConfiguration& InConfiguration,
 		UObject* Outer) const { return nullptr; }
 
-	virtual FText ToText(const FMediaIOConfiguration& InConfiguration) const;
+	virtual FText ToText(const FMediaIOConfiguration& InConfiguration, bool bInIsAutoDetected = false) const;
 	virtual FText ToText(const FMediaIOConnection& InConnection) const;
 	virtual FText ToText(const FMediaIODevice& InDevice) const;
 	virtual FText ToText(const FMediaIOMode& InMode) const;
diff --git a/Engine/Plugins/Media/MediaIOFramework/Source/MediaIOEditor/Private/Customizations/MediaIOConfigurationCustomization.cpp b/Engine/Plugins/Media/MediaIOFramework/Source/MediaIOEditor/Private/Customizations/MediaIOConfigurationCustomization.cpp
--- a/Engine/Plugins/Media/MediaIOFramework/Source/MediaIOEditor/Private/Customizations/MediaIOConfigurationCustomization.cpp
+++ b/Engine/Plugins/Media/MediaIOFramework/Source/MediaIOEditor/Private/Customizations/MediaIOConfigurationCustomization.cpp
@@ -35,7 +35,10 @@ TAttribute<FText> FMediaIOConfigurationCustomization::GetContentText()
 	IMediaIOCoreDeviceProvider* DeviceProviderPtr = IMediaIOCoreModule::Get().GetDeviceProvider(DeviceProviderName);
 	if (DeviceProviderPtr)
 	{
-		return MakeAttributeLambda([=] { return DeviceProviderPtr->ToText(*Value); });
+		return MakeAttributeLambda([this, DeviceProviderPtr, Value]
+			{ 
+				return DeviceProviderPtr->ToText(*Value, IsAutoDetected()); 
+			});
 	}
 	return FText::GetEmpty();
 }
@@ -59,7 +62,7 @@ TSharedRef<SWidget> FMediaIOConfigurationCustomization::HandleSourceComboButtonM
 		SelectedConfiguration.bIsInput = bIsInput;
 	}
 
-	bEnforceFormat = !IsAutoDetected();
+	bAutoDetectFormat = IsAutoDetected();
 
 	TArray<FMediaIOConfiguration> MediaConfigurations = bIsInput ? DeviceProviderPtr->GetConfigurations(true, false) : DeviceProviderPtr->GetConfigurations(false, true);
 	if (MediaConfigurations.Num() == 0)
@@ -77,7 +80,7 @@ TSharedRef<SWidget> FMediaIOConfigurationCustomization::HandleSourceComboButtonM
 		return false;
 	};
 
-	TSharedRef<SWidget> EnforceWidget = SNullWidget::NullWidget;
+	TSharedRef<SWidget> AutoCheckbox = SNullWidget::NullWidget;
 
 	if (bIsInput)
 	{
@@ -106,32 +109,32 @@ TSharedRef<SWidget> FMediaIOConfigurationCustomization::HandleSourceComboButtonM
 
 		if (bSupportsAutoDetect)
 		{
-			EnforceWidget = SNew(SHorizontalBox)
+			AutoCheckbox = SNew(SHorizontalBox)
 				+ SHorizontalBox::Slot()
 				.Padding(4.f)
 				.AutoWidth()
 				[
 					SNew(STextBlock)
-					.Text(NSLOCTEXT("MediaPlayerEditor", "EnforceLabel", "Enforce format"))
+					.Text(NSLOCTEXT("MediaPlayerEditor", "AutoLabel", "Auto"))
 				]
 			+ SHorizontalBox::Slot()
 				.Padding(4.f)
 				.AutoWidth()
 				[
 					SNew(SCheckBox)
-					.ToolTipText(NSLOCTEXT("MediaPlayerEditor", "EnforceTooltip", "Enforce a specific configuration. (Note that this disables format autodetection)"))
-					.IsChecked(this, &FMediaIOConfigurationCustomization::GetEnforceCheckboxState)
-					.OnCheckStateChanged(this, &FMediaIOConfigurationCustomization::SetEnforceCheckboxState)
+					.ToolTipText(NSLOCTEXT("MediaPlayerEditor", "AutoToolTip", "Automatically detect the video format of the signal coming through this input. When disabled, the format will be enforced and the source will error out if the format differs from the expected one."))
+					.IsChecked(this, &FMediaIOConfigurationCustomization::GetAutoCheckboxState)
+					.OnCheckStateChanged(this, &FMediaIOConfigurationCustomization::SetAutoCheckboxState)
 				];
 		}		
 
 	}
 
-	auto GetExtensions = [EnforceWidget](TArray<TSharedRef<SWidget>>& OutWidgets)
+	auto GetExtensions = [AutoCheckbox](TArray<TSharedRef<SWidget>>& OutWidgets)
 	{
-		if (EnforceWidget != SNullWidget::NullWidget)
+		if (AutoCheckbox != SNullWidget::NullWidget)
 		{
-			OutWidgets.Add(EnforceWidget);
+			OutWidgets.Add(AutoCheckbox);
 		}
 	};
 
@@ -185,7 +188,7 @@ FReply FMediaIOConfigurationCustomization::OnButtonClicked()
 {
 	AssignValue(SelectedConfiguration);
 	// Make sure to overwrite what was in the config since the auto value is determined by the timecode provider and not the generated configs.
-	SetIsAutoDetected(!bEnforceFormat);
+	SetIsAutoDetected(bAutoDetectFormat);
 
 	TSharedPtr<SWidget> SharedPermutationSelector = PermutationSelector.Pin();
 	if (SharedPermutationSelector.IsValid())
@@ -197,19 +200,19 @@ FReply FMediaIOConfigurationCustomization::OnButtonClicked()
 	return FReply::Handled();
 }
 
-ECheckBoxState FMediaIOConfigurationCustomization::GetEnforceCheckboxState() const
+ECheckBoxState FMediaIOConfigurationCustomization::GetAutoCheckboxState() const
 {
-	return bEnforceFormat ? ECheckBoxState::Checked : ECheckBoxState::Unchecked;
+	return bAutoDetectFormat ? ECheckBoxState::Checked : ECheckBoxState::Unchecked;
 }
 
-void FMediaIOConfigurationCustomization::SetEnforceCheckboxState(ECheckBoxState CheckboxState)
+void FMediaIOConfigurationCustomization::SetAutoCheckboxState(ECheckBoxState CheckboxState)
 {
-	bEnforceFormat = CheckboxState == ECheckBoxState::Checked;
+	bAutoDetectFormat = CheckboxState == ECheckBoxState::Checked;
 }
 
 bool FMediaIOConfigurationCustomization::ShowAdvancedColumns(FName ColumnName, const TArray<FMediaIOConfiguration>& UniquePermutationsForThisColumn) const
 {
-	return bEnforceFormat;
+	return !bAutoDetectFormat;
 }
 
 bool FMediaIOConfigurationCustomization::IsAutoDetected() const
diff --git a/Engine/Plugins/Media/MediaIOFramework/Source/MediaIOEditor/Private/Customizations/MediaIOVideoTimecodeConfigurationCustomization.cpp b/Engine/Plugins/Media/MediaIOFramework/Source/MediaIOEditor/Private/Customizations/MediaIOVideoTimecodeConfigurationCustomization.cpp
--- a/Engine/Plugins/Media/MediaIOFramework/Source/MediaIOEditor/Private/Customizations/MediaIOVideoTimecodeConfigurationCustomization.cpp
+++ b/Engine/Plugins/Media/MediaIOFramework/Source/MediaIOEditor/Private/Customizations/MediaIOVideoTimecodeConfigurationCustomization.cpp
@@ -89,7 +89,7 @@ TSharedRef<SWidget> FMediaIOVideoTimecodeConfigurationCustomization::HandleSourc
 		}
 	}
 
-	bEnforceFormat = !IsAutoDetected();
+	bAutoDetectFormat = IsAutoDetected();
 
 	TArray<FMediaIOVideoTimecodeConfiguration> MediaConfigurations;
 	if (IMediaIOCoreDeviceProvider* DeviceProvider = IMediaIOCoreModule::Get().GetDeviceProvider(DeviceProviderName))
@@ -112,27 +112,27 @@ TSharedRef<SWidget> FMediaIOVideoTimecodeConfigurationCustomization::HandleSourc
 		return false;
 	};
 
-	TSharedRef<SWidget> EnforceWidget = SNew(SHorizontalBox)
+	TSharedRef<SWidget> AutoCheckbox = SNew(SHorizontalBox)
 		+ SHorizontalBox::Slot()
 		.Padding(4.f)
 		.AutoWidth()
 		[
 			SNew(STextBlock)
-			.Text(LOCTEXT("EnforceLabel", "Enforce format"))
+			.Text(LOCTEXT("AutoLabel", "Auto"))
 		]
 	+ SHorizontalBox::Slot()
 		.Padding(4.f)
 		.AutoWidth()
 		[
 			SNew(SCheckBox)
-			.ToolTipText(LOCTEXT("EnforceTooltip", "Enforce a specific configuration. (Note that this disables format autodetection)"))
-			.IsChecked(this, &FMediaIOVideoTimecodeConfigurationCustomization::GetEnforceCheckboxState)
-			.OnCheckStateChanged(this, &FMediaIOVideoTimecodeConfigurationCustomization::SetEnforceCheckboxState)
+			.ToolTipText(LOCTEXT("AutoTooltip", "Automatically detect the video format of the timecode signal coming through this input. When disabled, the format will be enforced and the source will error out if the format differs from the expected one."))
+			.IsChecked(this, &FMediaIOVideoTimecodeConfigurationCustomization::GetAutoCheckboxState)
+			.OnCheckStateChanged(this, &FMediaIOVideoTimecodeConfigurationCustomization::SetAutoCheckboxState)
 		];
 
-	auto GetExtensions = [EnforceWidget](TArray<TSharedRef<SWidget>>& OutWidgets)
+	auto GetExtensions = [AutoCheckbox](TArray<TSharedRef<SWidget>>& OutWidgets)
 	{
-		OutWidgets.Add(EnforceWidget);
+		OutWidgets.Add(AutoCheckbox);
 	};
 
 	using TSelection = SMediaPermutationsSelector<FMediaIOVideoTimecodeConfiguration, MediaIOVideoTimecodeConfigurationCustomization::FMediaTimecodePermutationsSelectorBuilder>;
@@ -191,7 +191,7 @@ FReply FMediaIOVideoTimecodeConfigurationCustomization::OnButtonClicked()
 {
 	AssignValue(SelectedConfiguration);
 	// Make sure to overwrite what was in the config since the auto value is determined by the timecode provider and not the generated configs.
-	SetIsAutoDetected(!bEnforceFormat);
+	SetIsAutoDetected(bAutoDetectFormat);
 
 	TSharedPtr<SWidget> SharedPermutationSelector = PermutationSelector.Pin();
 	if (SharedPermutationSelector.IsValid())
@@ -203,19 +203,19 @@ FReply FMediaIOVideoTimecodeConfigurationCustomization::OnButtonClicked()
 	return FReply::Handled();
 }
 
-ECheckBoxState FMediaIOVideoTimecodeConfigurationCustomization::GetEnforceCheckboxState() const
+ECheckBoxState FMediaIOVideoTimecodeConfigurationCustomization::GetAutoCheckboxState() const
 {
-	return bEnforceFormat ? ECheckBoxState::Checked : ECheckBoxState::Unchecked;
+	return bAutoDetectFormat ? ECheckBoxState::Checked : ECheckBoxState::Unchecked;
 }
 
-void FMediaIOVideoTimecodeConfigurationCustomization::SetEnforceCheckboxState(ECheckBoxState CheckboxState)
+void FMediaIOVideoTimecodeConfigurationCustomization::SetAutoCheckboxState(ECheckBoxState CheckboxState)
 {
-	bEnforceFormat = CheckboxState == ECheckBoxState::Checked;
+	bAutoDetectFormat = CheckboxState == ECheckBoxState::Checked;
 }
 
 bool FMediaIOVideoTimecodeConfigurationCustomization::ShowAdvancedColumns(FName ColumnName, const TArray<FMediaIOVideoTimecodeConfiguration>& UniquePermutationsForThisColumn) const
 {
-	return bEnforceFormat;
+	return !bAutoDetectFormat;
 }
 
 bool FMediaIOVideoTimecodeConfigurationCustomization::IsAutoDetected() const
diff --git a/Engine/Plugins/Media/MediaIOFramework/Source/MediaIOEditor/Public/Customizations/MediaIOConfigurationCustomization.h b/Engine/Plugins/Media/MediaIOFramework/Source/MediaIOEditor/Public/Customizations/MediaIOConfigurationCustomization.h
--- a/Engine/Plugins/Media/MediaIOFramework/Source/MediaIOEditor/Public/Customizations/MediaIOConfigurationCustomization.h
+++ b/Engine/Plugins/Media/MediaIOFramework/Source/MediaIOEditor/Public/Customizations/MediaIOConfigurationCustomization.h
@@ -18,8 +18,8 @@ private:
 	virtual TAttribute<FText> GetContentText() override;
 	virtual TSharedRef<SWidget> HandleSourceComboButtonMenuContent() override;
 
-	ECheckBoxState GetEnforceCheckboxState() const;
-	void SetEnforceCheckboxState(ECheckBoxState CheckboxState);
+	ECheckBoxState GetAutoCheckboxState() const;
+	void SetAutoCheckboxState(ECheckBoxState CheckboxState);
 	void OnSelectionChanged(FMediaIOConfiguration SelectedItem);
 	FReply OnButtonClicked();
 	bool ShowAdvancedColumns(FName ColumnName, const TArray<FMediaIOConfiguration>& UniquePermutationsForThisColumn) const;
@@ -29,5 +29,5 @@ private:
 private:
 	TWeakPtr<SWidget> PermutationSelector;
 	FMediaIOConfiguration SelectedConfiguration;
-	bool bEnforceFormat = true;
+	bool bAutoDetectFormat = true;
 };
diff --git a/Engine/Plugins/Media/MediaIOFramework/Source/MediaIOEditor/Public/Customizations/MediaIOVideoTimecodeConfigurationCustomization.h b/Engine/Plugins/Media/MediaIOFramework/Source/MediaIOEditor/Public/Customizations/MediaIOVideoTimecodeConfigurationCustomization.h
--- a/Engine/Plugins/Media/MediaIOFramework/Source/MediaIOEditor/Public/Customizations/MediaIOVideoTimecodeConfigurationCustomization.h
+++ b/Engine/Plugins/Media/MediaIOFramework/Source/MediaIOEditor/Public/Customizations/MediaIOVideoTimecodeConfigurationCustomization.h
@@ -20,8 +20,8 @@ private:
 
 	void OnSelectionChanged(FMediaIOVideoTimecodeConfiguration SelectedItem);
 	FReply OnButtonClicked();
-	ECheckBoxState GetEnforceCheckboxState() const;
-	void SetEnforceCheckboxState(ECheckBoxState CheckboxState);
+	ECheckBoxState GetAutoCheckboxState() const;
+	void SetAutoCheckboxState(ECheckBoxState CheckboxState);
 	bool ShowAdvancedColumns(FName ColumnName, const TArray<FMediaIOVideoTimecodeConfiguration>& UniquePermutationsForThisColumn) const;
 	bool IsAutoDetected() const;
 	void SetIsAutoDetected(bool Value);
@@ -29,5 +29,5 @@ private:
 private:
 	TWeakPtr<SWidget> PermutationSelector;
 	FMediaIOVideoTimecodeConfiguration SelectedConfiguration;
-	bool bEnforceFormat = false;
+	bool bAutoDetectFormat = false;
 };
