Fix DVP being initialized on unsupported GPUs

diff --git a/Engine/Plugins/Media/AjaMedia/Source/AjaMediaOutput/Private/AjaMediaOutputModule.cpp b/Engine/Plugins/Media/AjaMedia/Source/AjaMediaOutput/Private/AjaMediaOutputModule.cpp
--- a/Engine/Plugins/Media/AjaMedia/Source/AjaMediaOutput/Private/AjaMediaOutputModule.cpp
+++ b/Engine/Plugins/Media/AjaMedia/Source/AjaMediaOutput/Private/AjaMediaOutputModule.cpp
@@ -18,9 +18,29 @@ void FAjaMediaOutputModule::StartupModule()
 {
 	const auto DMAInitializationFunc = [this]()
 	{
+		static const TArray<FString> SupportedGPUPrefixes= { 
+			TEXT("RTX A4"),
+			TEXT("RTX A5"),
+			TEXT("RTX A6"),
+			TEXT("Quadro")
+		};
+
 		const FGPUDriverInfo GPUDriverInfo = FPlatformMisc::GetGPUDriverInfo(GRHIAdapterName);
 		bIsGPUTextureTransferAvailable = GPUDriverInfo.IsNVIDIA() &&!FModuleManager::Get().IsModuleLoaded("RenderDocPlugin");
 		bIsGPUTextureTransferAvailable &= !GPUDriverInfo.DeviceDescription.Contains(TEXT("Tesla"));
+		
+		if (bIsGPUTextureTransferAvailable)
+		{
+			bIsGPUTextureTransferAvailable = false;
+			for (const FString& GPUPrefix : SupportedGPUPrefixes)
+			{
+				if (GPUDriverInfo.DeviceDescription.Contains(GPUPrefix))
+				{
+					bIsGPUTextureTransferAvailable = true;
+					break;
+				}
+			}
+		}
 
 		if (bIsGPUTextureTransferAvailable)
 		{
@@ -58,6 +78,7 @@ void FAjaMediaOutputModule::StartupModule()
 					Args.RHIDevice = GDynamicRHI->RHIGetNativeDevice();
 					Args.RHICommandQueue = GDynamicRHI->RHIGetNativeGraphicsQueue();
 
+					UE_LOG(LogAjaMediaOutput, Display, TEXT("Initializing GPU Texture transfer from AJA"));
 					bIsGPUTextureTransferAvailable = AJA::InitializeDMA(Args);
 				});
 		}
diff --git a/Engine/Plugins/Media/BlackmagicMedia/Source/BlackmagicMediaOutput/Private/BlackmagicMediaOutputModule.cpp b/Engine/Plugins/Media/BlackmagicMedia/Source/BlackmagicMediaOutput/Private/BlackmagicMediaOutputModule.cpp
--- a/Engine/Plugins/Media/BlackmagicMedia/Source/BlackmagicMediaOutput/Private/BlackmagicMediaOutputModule.cpp
+++ b/Engine/Plugins/Media/BlackmagicMedia/Source/BlackmagicMediaOutput/Private/BlackmagicMediaOutputModule.cpp
@@ -18,10 +18,30 @@ void FBlackmagicMediaOutputModule::StartupModule()
 {
 	const auto DMAInitializationFunc = [this]()
 	{
+		static const TArray<FString> SupportedGPUPrefixes = {
+			TEXT("RTX A4"),
+			TEXT("RTX A5"),
+			TEXT("RTX A6"),
+			TEXT("Quadro")
+		};
+
 		const FGPUDriverInfo GPUDriverInfo = FPlatformMisc::GetGPUDriverInfo(GRHIAdapterName);
 		bIsGPUTextureTransferAvailable = GPUDriverInfo.IsNVIDIA() && !FModuleManager::Get().IsModuleLoaded("RenderDocPlugin");
 		bIsGPUTextureTransferAvailable &= !GPUDriverInfo.DeviceDescription.Contains(TEXT("Tesla"));
 
+		if (bIsGPUTextureTransferAvailable)
+		{
+			bIsGPUTextureTransferAvailable = false;
+			for (const FString& GPUPrefix : SupportedGPUPrefixes)
+			{
+				if (GPUDriverInfo.DeviceDescription.Contains(GPUPrefix))
+				{
+					bIsGPUTextureTransferAvailable = true;
+					break;
+				}
+			}
+		}
+
 		if (bIsGPUTextureTransferAvailable)
 		{
 			ENQUEUE_RENDER_COMMAND(BlackmagicMediaCaptureInitialize)(
diff --git a/Engine/Plugins/Media/MediaIOFramework/Source/GPUTextureTransfer/Private/GPUTextureTransferModule.cpp b/Engine/Plugins/Media/MediaIOFramework/Source/GPUTextureTransfer/Private/GPUTextureTransferModule.cpp
--- a/Engine/Plugins/Media/MediaIOFramework/Source/GPUTextureTransfer/Private/GPUTextureTransferModule.cpp
+++ b/Engine/Plugins/Media/MediaIOFramework/Source/GPUTextureTransfer/Private/GPUTextureTransferModule.cpp
@@ -138,11 +138,30 @@ void FGPUTextureTransferModule::InitializeTextureTransfer()
 {
 #if DVP_SUPPORTED_PLATFORM
 	bIsGPUTextureTransferAvailable = true;
+
+	static const TArray<FString> SupportedGPUPrefixes = {
+		TEXT("RTX A4"),
+		TEXT("RTX A5"),
+		TEXT("RTX A6"),
+		TEXT("Quadro")
+	};
+
 	// This must be called on game thread 
 	const FGPUDriverInfo GPUDriverInfo = FPlatformMisc::GetGPUDriverInfo(GRHIAdapterName);
-	bIsGPUTextureTransferAvailable = GPUDriverInfo.IsNVIDIA() && !FModuleManager::Get().IsModuleLoaded("RenderDocPlugin");
-	bIsGPUTextureTransferAvailable = bIsGPUTextureTransferAvailable && !GPUDriverInfo.DeviceDescription.Contains(TEXT("Tesla"));
+	bIsGPUTextureTransferAvailable = GPUDriverInfo.IsNVIDIA() && !FModuleManager::Get().IsModuleLoaded("RenderDocPlugin") && !GPUDriverInfo.DeviceDescription.Contains(TEXT("Tesla"));
 
+	if (bIsGPUTextureTransferAvailable)
+	{
+		bIsGPUTextureTransferAvailable = false;
+		for (const FString& GPUPrefix : SupportedGPUPrefixes)
+		{
+			if (GPUDriverInfo.DeviceDescription.Contains(GPUPrefix))
+			{
+				bIsGPUTextureTransferAvailable = true;
+				break;
+			}
+		}
+	}
 	if (!bIsGPUTextureTransferAvailable)
 	{
 		return;
@@ -190,6 +209,7 @@ void FGPUTextureTransferModule::InitializeTextureTransfer()
 #endif
 
 		const uint8 RHIIndex = static_cast<uint8>(RHI);
+		UE_LOG(LogGPUTextureTransfer, Display, TEXT("Initializing GPU Texture transfer"));
 		if (TextureTransfer->Initialize(InitializeArgs))
 		{
 			TransferObjects[RHIIndex] = TextureTransfer;
