[UE]
* More fixes to be preprocessor conformant (/Zc:preprocessor)

diff --git a/Engine/Plugins/Media/MediaIOFramework/Source/GPUTextureTransfer/Private/TextureTransferBase.cpp b/Engine/Plugins/Media/MediaIOFramework/Source/GPUTextureTransfer/Private/TextureTransferBase.cpp
--- a/Engine/Plugins/Media/MediaIOFramework/Source/GPUTextureTransfer/Private/TextureTransferBase.cpp
+++ b/Engine/Plugins/Media/MediaIOFramework/Source/GPUTextureTransfer/Private/TextureTransferBase.cpp
@@ -29,7 +29,7 @@ ON_SCOPE_EXIT{ UE_LOG(LogGPUTextureTransfer, Verbose, TEXT("%s duration: %d ms")
 	DVPStatus _Status = (call); \
 	if (_Status != DVP_STATUS_OK)\
 	{\
-		UE_LOG(LogGPUTextureTransfer, Error, TEXT("GPUDirect call %s failed. Error:  %d"), L#call, _Status);\
+		UE_LOG(LogGPUTextureTransfer, Error, TEXT("GPUDirect call %s failed. Error:  %d"), L""#call, _Status);\
 	}\
 }
 
diff --git a/Engine/Source/Developer/Windows/LiveCoding/Private/External/LC_Assert.h b/Engine/Source/Developer/Windows/LiveCoding/Private/External/LC_Assert.h
--- a/Engine/Source/Developer/Windows/LiveCoding/Private/External/LC_Assert.h
+++ b/Engine/Source/Developer/Windows/LiveCoding/Private/External/LC_Assert.h
@@ -13,7 +13,7 @@
 
 #if LC_DEBUG
 #	define LC_ASSERT_BREAKPOINT()				(Debugger::IsConnected()) ? LC_DEBUGGER_BREAKPOINT() : (void)true
-#	define LC_ASSERT(_condition, _msg)			(_condition) ? (void)true : (LC_ERROR_DEV(_msg), LC_ASSERT_BREAKPOINT())
+#	define LC_ASSERT(_condition, _msg)			(_condition) ? (void)true : (LC_ERROR_DEV("%s", _msg), LC_ASSERT_BREAKPOINT())
 #else
 #	define LC_BREAKPOINT()						(void)true
 #	define LC_ASSERT(_condition, _msg)			LC_UNUSED(_condition)
