- Fixing odd errors when starting rivermax output stream
- Fixing missing call in media capture that were removed by accident

diff --git a/Engine/Plugins/Media/AjaMedia/Source/AjaMediaOutput/Private/AjaMediaOutputModule.cpp b/Engine/Plugins/Media/AjaMedia/Source/AjaMediaOutput/Private/AjaMediaOutputModule.cpp
--- a/Engine/Plugins/Media/AjaMedia/Source/AjaMediaOutput/Private/AjaMediaOutputModule.cpp
+++ b/Engine/Plugins/Media/AjaMedia/Source/AjaMediaOutput/Private/AjaMediaOutputModule.cpp
@@ -88,7 +88,7 @@ void FAjaMediaOutputModule::StartupModule()
 	{
 		if (bIsGPUTextureTransferAvailable)
 		{
-			ENQUEUE_RENDER_COMMAND(AjaMediaCaptureUninitialize)(
+			ENQUEUE_RENDER_COMMAND(AjaMediaOutputModuleUninitialize)(
 				[](FRHICommandListImmediate& RHICmdList) mutable
 				{
 					AJA::UninitializeDMA();
diff --git a/Engine/Plugins/Media/MediaIOFramework/Source/MediaIOCore/Private/MediaCapture.cpp b/Engine/Plugins/Media/MediaIOFramework/Source/MediaIOCore/Private/MediaCapture.cpp
--- a/Engine/Plugins/Media/MediaIOFramework/Source/MediaIOCore/Private/MediaCapture.cpp
+++ b/Engine/Plugins/Media/MediaIOFramework/Source/MediaIOCore/Private/MediaCapture.cpp
@@ -852,13 +852,6 @@ bool UMediaCapture::CaptureSceneViewport(TSharedPtr<FSceneViewport>& InSceneView
 
 	// This could have been updated by the initialization done by the implementation
 	bShouldCaptureRHIResource = ShouldCaptureRHIResource(); 
-	if (bShouldCaptureRHIResource && DesiredOutputResourceType != EMediaCaptureResourceType::Texture)
-	{
-		UE_LOG(LogMediaIOCore, Warning, TEXT("Can't capture RHI resource when not using texture output resource"));
-		ResetFixedViewportSize(InSceneViewport, false);
-		SetState(EMediaCaptureState::Stopped);
-		MediaCaptureDetails::ShowSlateNotification();
-	}
 
 	if (bInitialized)
 	{
@@ -921,12 +914,6 @@ bool UMediaCapture::CaptureRHITexture(const FRHICaptureResourceDescription& Reso
 
 	// This could have been updated by the call to implementation initialization
 	bShouldCaptureRHIResource = ShouldCaptureRHIResource();
-	if (bShouldCaptureRHIResource && DesiredOutputResourceType != EMediaCaptureResourceType::Texture)
-	{
-		UE_LOG(LogMediaIOCore, Warning, TEXT("Can't capture RHI resource when not using texture output resource"));
-		SetState(EMediaCaptureState::Stopped);
-		MediaCaptureDetails::ShowSlateNotification();
-	}
 
 	if (bInitialized)
 	{
@@ -996,12 +983,6 @@ bool UMediaCapture::CaptureTextureRenderTarget2D(UTextureRenderTarget2D* InRende
 
 	// This could have been updated by the initialization done by the implementation
 	bShouldCaptureRHIResource = ShouldCaptureRHIResource();
-	if (bShouldCaptureRHIResource && DesiredOutputResourceType != EMediaCaptureResourceType::Texture)
-	{
-		UE_LOG(LogMediaIOCore, Warning, TEXT("Can't capture RHI resource when not using texture output resource"));
-		SetState(EMediaCaptureState::Stopped);
-		MediaCaptureDetails::ShowSlateNotification();
-	}
 
 	if (bInitialized)
 	{
@@ -1745,6 +1726,7 @@ bool UMediaCapture::ProcessCapture_RenderThread(FRDGBuilder& GraphBuilder, UMedi
 			if (ensure(CapturingFrame->IsTextureResource()))
 			{
 				FTextureCaptureFrame* TextureFrame = static_cast<FTextureCaptureFrame*>(CapturingFrame);
+				InMediaCapture->BeforeFrameCaptured_RenderingThread(CapturingFrame->CaptureBaseData, CapturingFrame->UserData, TextureFrame->GetTextureResource());
 				bHasCaptureSuceeded = FMediaCaptureHelper::CaptureFrame(CaptureArgs, TextureFrame);
 			}
 			else
@@ -1757,6 +1739,7 @@ bool UMediaCapture::ProcessCapture_RenderThread(FRDGBuilder& GraphBuilder, UMedi
 			if (ensure(CapturingFrame->IsBufferResource()))
 			{
 				FBufferCaptureFrame* BufferFrame = static_cast<FBufferCaptureFrame*>(CapturingFrame);
+				InMediaCapture->BeforeFrameCaptured_RenderingThread(CapturingFrame->CaptureBaseData, CapturingFrame->UserData, BufferFrame->GetBufferResource());
 				bHasCaptureSuceeded = FMediaCaptureHelper::CaptureFrame(CaptureArgs, BufferFrame);
 			}
 			else
@@ -1850,13 +1833,19 @@ bool UMediaCapture::ProcessReadyFrame_RenderThread(FRHICommandListImmediate& RHI
 
 				TRACE_CPUPROFILER_EVENT_SCOPE(UMediaCapture::RHIResourceCaptured);
 				SCOPE_CYCLE_COUNTER(STAT_MediaCapture_RenderThread_CaptureCallback)
-				InMediaCapture->OnRHIResourceCaptured_RenderingThread(ReadyFrame->CaptureBaseData, ReadyFrame->UserData, TextureFrame->GetRHIResource());
-
-				UE_LOG(LogMediaIOCore, Verbose, TEXT("[%s] - Completed pending frame %d."), *InMediaCapture->MediaOutputName, ReadyFrame->FrameId);
-				ReadyFrame->bDoingGPUCopy = false;
-				InMediaCapture->FrameManager->CompleteNextPending(*ReadyFrame);
-				--PendingFrameCount;
+					InMediaCapture->OnRHIResourceCaptured_RenderingThread(ReadyFrame->CaptureBaseData, ReadyFrame->UserData, TextureFrame->GetRHIResource());
 			}
+			else
+			{
+				TRACE_CPUPROFILER_EVENT_SCOPE(UMediaCapture::RHIResourceCaptured);
+				FBufferCaptureFrame* BufferFrame = static_cast<FBufferCaptureFrame*>(ReadyFrame);
+				InMediaCapture->OnRHIResourceCaptured_RenderingThread(ReadyFrame->CaptureBaseData, ReadyFrame->UserData, BufferFrame->GetRHIResource());
+			}
+
+			UE_LOG(LogMediaIOCore, Verbose, TEXT("[%s] - Completed pending frame %d."), *InMediaCapture->MediaOutputName, ReadyFrame->FrameId);
+			ReadyFrame->bDoingGPUCopy = false;
+			InMediaCapture->FrameManager->CompleteNextPending(*ReadyFrame);
+			--PendingFrameCount;
 		}
 	}
 
diff --git a/Engine/Plugins/VirtualProduction/Rivermax/RivermaxCore/Source/RivermaxCore/Private/RivermaxInputStream.cpp b/Engine/Plugins/VirtualProduction/Rivermax/RivermaxCore/Source/RivermaxCore/Private/RivermaxInputStream.cpp
--- a/Engine/Plugins/VirtualProduction/Rivermax/RivermaxCore/Source/RivermaxCore/Private/RivermaxInputStream.cpp
+++ b/Engine/Plugins/VirtualProduction/Rivermax/RivermaxCore/Source/RivermaxCore/Private/RivermaxInputStream.cpp
@@ -235,6 +235,7 @@ namespace UE::RivermaxCore::Private
 		const int Timeout = 0;
 		const int Flags = 0;
 		rmax_in_completion Completion;
+		FMemory::Memset(&Completion, 0, sizeof(Completion));
 		rmax_status_t Status = rmax_in_get_next_chunk(StreamId, MinChunkSize, MaxChunkSize, Timeout, Flags, &Completion);
 		if (Status == RMAX_OK)
 		{
@@ -497,7 +498,3 @@ namespace UE::RivermaxCore::Private
 
 }
 
-
-
-
-
diff --git a/Engine/Plugins/VirtualProduction/Rivermax/RivermaxCore/Source/RivermaxCore/Private/RivermaxOutputStream.cpp b/Engine/Plugins/VirtualProduction/Rivermax/RivermaxCore/Source/RivermaxCore/Private/RivermaxOutputStream.cpp
--- a/Engine/Plugins/VirtualProduction/Rivermax/RivermaxCore/Source/RivermaxCore/Private/RivermaxOutputStream.cpp
+++ b/Engine/Plugins/VirtualProduction/Rivermax/RivermaxCore/Source/RivermaxCore/Private/RivermaxOutputStream.cpp
@@ -148,17 +148,21 @@ namespace UE::RivermaxCore::Private
 				rmax_qos_attr QOSAttributes = { 0, 0 }; //todo
 
 				rmax_buffer_attr BufferAttributes;
+				FMemory::Memset(&BufferAttributes, 0, sizeof(BufferAttributes));
 				BufferAttributes.chunk_size_in_strides = StreamMemory.ChunkSizeInStrides;
 				BufferAttributes.data_stride_size = StreamMemory.PayloadSize; //Stride between chunks. 
 				BufferAttributes.app_hdr_stride_size = StreamMemory.HeaderStrideSize;
 				BufferAttributes.mem_block_array = StreamMemory.MemoryBlocks.GetData();
 				BufferAttributes.mem_block_array_len = StreamMemory.MemoryBlocks.Num();
+				BufferAttributes.attr_flags = RMAX_OUT_BUFFER_ATTR_FLAG_NONE;
 
 				rmax_status_t Status = rmax_out_create_stream(SDPDescription.GetData(), &BufferAttributes, &QOSAttributes, NumberPacketsPerFrame, MediaBlockIndex, &NewId);
 				if (Status == RMAX_OK)
 				{
 					struct sockaddr_in SourceAddress;
 					struct sockaddr_in DestinationAddress;
+					FMemory::Memset(&SourceAddress, 0, sizeof(SourceAddress));
+					FMemory::Memset(&DestinationAddress, 0, sizeof(DestinationAddress));
 					Status = rmax_out_query_address(NewId, MediaBlockIndex, &SourceAddress, &DestinationAddress);
 					if (Status == RMAX_OK)
 					{
@@ -363,7 +367,7 @@ namespace UE::RivermaxCore::Private
 			}
 			StreamMemory.HeaderSizes[PayloadSizeIndex] = StreamMemory.HeaderStrideSize;
 		}
-		StreamMemory.MemoryBlocks.SetNum(StreamMemory.MemoryBlockCount);
+		StreamMemory.MemoryBlocks.SetNumZeroed(StreamMemory.MemoryBlockCount);
 		for (uint32 BlockIndex = 0; BlockIndex < StreamMemory.MemoryBlockCount; ++BlockIndex)
 		{
 			rmax_mem_block& Block = StreamMemory.MemoryBlocks[BlockIndex];
