Fix audio not working on aja inputs

diff --git a/Engine/Plugins/Media/MediaIOFramework/Source/MediaIOCore/Private/Shared/MediaIOCoreAudioSampleBase.cpp b/Engine/Plugins/Media/MediaIOFramework/Source/MediaIOCore/Private/Shared/MediaIOCoreAudioSampleBase.cpp
--- a/Engine/Plugins/Media/MediaIOFramework/Source/MediaIOCore/Private/Shared/MediaIOCoreAudioSampleBase.cpp
+++ b/Engine/Plugins/Media/MediaIOFramework/Source/MediaIOCore/Private/Shared/MediaIOCoreAudioSampleBase.cpp
@@ -76,6 +76,9 @@ bool FMediaIOCoreAudioSampleBase::SetProperties(uint32 InBufferSize, uint32 InNu
 	SampleRate = InSampleRate;
 	Duration = (InBufferSize * ETimespan::TicksPerSecond) / (Channels * SampleRate);
 
+	constexpr bool bAllowShrinking = false;
+	Buffer.SetNum(InBufferSize, bAllowShrinking);
+
 	return true;
 }
 
