Fix CIS in GPUTextureTransfer

diff --git a/Engine/Plugins/Media/MediaIOFramework/Source/GPUTextureTransfer/Private/GPUTextureTransferModule.cpp b/Engine/Plugins/Media/MediaIOFramework/Source/GPUTextureTransfer/Private/GPUTextureTransferModule.cpp
--- a/Engine/Plugins/Media/MediaIOFramework/Source/GPUTextureTransfer/Private/GPUTextureTransferModule.cpp
+++ b/Engine/Plugins/Media/MediaIOFramework/Source/GPUTextureTransfer/Private/GPUTextureTransferModule.cpp
@@ -243,10 +243,12 @@ namespace UE::GPUTextureTransfer
 	struct FTextureTransfersWrapper
 	{
 		static constexpr uint8_t RHI_MAX = static_cast<uint8_t>(ERHI::RHI_MAX);
-		ITextureTransfer* Transfers[RHI_MAX];
+		TArray<ITextureTransfer*> Transfers;
 
 		FTextureTransfersWrapper()
 		{
+			Transfers.SetNumUninitialized(RHI_MAX);
+
 			for (uint8_t RhiIt = 0; RhiIt < RHI_MAX; RhiIt++)
 			{
 				Transfers[RhiIt] = nullptr;
@@ -274,10 +276,10 @@ namespace UE::GPUTextureTransfer
 			{
 				for (uint8_t RhiIt = 1; RhiIt < RHI_MAX; RhiIt++)
 				{
-					if (Transfers[RhiIt] == TextureTransfer)
+					if (Transfers[RhiIt] && Transfers[RhiIt] == TextureTransfer)
 					{
-						TextureTransfer->Uninitialize();
-						delete TextureTransfer;
+						Transfers[RhiIt]->Uninitialize();
+						delete Transfers[RhiIt];
 						Transfers[RhiIt] = nullptr;
 					}
 				}
