Reactivate time syncronization writer program

diff --git a/Engine/Plugins/Media/MediaIOFramework/Source/GPUTextureTransfer/GPUTextureTransfer.Build.cs b/Engine/Plugins/Media/MediaIOFramework/Source/GPUTextureTransfer/GPUTextureTransfer.Build.cs
--- a/Engine/Plugins/Media/MediaIOFramework/Source/GPUTextureTransfer/GPUTextureTransfer.Build.cs
+++ b/Engine/Plugins/Media/MediaIOFramework/Source/GPUTextureTransfer/GPUTextureTransfer.Build.cs
@@ -9,9 +9,17 @@ namespace UnrealBuildTool.Rules
 			if (Target.Platform == UnrealTargetPlatform.Win64)
 			{
 				AddEngineThirdPartyPrivateStaticDependencies(Target, "GPUDirect");
-				AddEngineThirdPartyPrivateStaticDependencies(Target, "Vulkan");
-				PrivateDependencyModuleNames.Add("VulkanRHI");
 				PublicDefinitions.Add("DVP_SUPPORTED_PLATFORM=1");
+
+				if (Target.bCompileAgainstEngine)
+				{
+					AddEngineThirdPartyPrivateStaticDependencies(Target, "Vulkan");
+					PrivateDependencyModuleNames.Add("VulkanRHI");
+				}
+				else
+				{
+					PublicDefinitions.Add("DVP_SUPPORTED_PLATFORM=0");
+				}
 			}
 			else if (Target.Platform == UnrealTargetPlatform.Linux || Target.Platform == UnrealTargetPlatform.LinuxArm64)
 			{
