Fix AJA not showing the port number when autodected

diff --git a/Engine/Plugins/Media/MediaIOFramework/Source/MediaIOCore/Private/MediaIOCoreDeviceProvider.cpp b/Engine/Plugins/Media/MediaIOFramework/Source/MediaIOCore/Private/MediaIOCoreDeviceProvider.cpp
--- a/Engine/Plugins/Media/MediaIOFramework/Source/MediaIOCore/Private/MediaIOCoreDeviceProvider.cpp
+++ b/Engine/Plugins/Media/MediaIOFramework/Source/MediaIOCore/Private/MediaIOCoreDeviceProvider.cpp
@@ -33,10 +33,12 @@ FText IMediaIOCoreDeviceProvider::ToText(const FMediaIOConfiguration& InConfigur
 {
 	if (bInIsAutoDetected)
 	{
-		return FText::Format(LOCTEXT("FMediaIOConfigurationToTextAutoDetect", "{0} - {1} [device{2}/auto]")
+		return FText::Format(LOCTEXT("FMediaIOConfigurationToTextAutoDetect", "{0} - {1} [device{2}/{3}{4}/auto]")
 				, InConfiguration.bIsInput ? LOCTEXT("In", "In") : LOCTEXT("Out", "Out")
 				, FText::FromName(InConfiguration.MediaConnection.Device.DeviceName)
 				, FText::AsNumber(InConfiguration.MediaConnection.Device.DeviceIdentifier)
+				, GetTransportName(InConfiguration.MediaConnection.TransportType, InConfiguration.MediaConnection.QuadTransportType)
+				, FText::AsNumber(InConfiguration.MediaConnection.PortIdentifier)
 				);
 	}
 	if (InConfiguration.IsValid())
