Disable GPU Direct in DX11, move initialisation to game thread
(Disabled in dx11 because there are some RHI calls being made which could be problematic on dx11 because of lack of good multithreading support)

diff --git a/Engine/Plugins/Media/AjaMedia/Source/Aja/Private/TextureDMA.cpp b/Engine/Plugins/Media/AjaMedia/Source/Aja/Private/TextureDMA.cpp
--- a/Engine/Plugins/Media/AjaMedia/Source/Aja/Private/TextureDMA.cpp
+++ b/Engine/Plugins/Media/AjaMedia/Source/Aja/Private/TextureDMA.cpp
@@ -46,7 +46,7 @@ namespace AJA
 	{
 		UE::GPUTextureTransfer::FInitializeDMAArgs GPUArgs;
 		GPUArgs.RHICommandQueue = Args.RHICommandQueue;
-
+		
 		switch (Args.RHI)
 		{
 		case ERHI::D3D11:
@@ -58,8 +58,10 @@ namespace AJA
 		case ERHI::Vulkan:
 			GPUArgs.RHI = UE::GPUTextureTransfer::ERHI::Vulkan;
 			break;
+		case ERHI::Invalid:
+			//Fallthrough
 		default:
-			checkNoEntry();
+			GPUArgs.RHI = UE::GPUTextureTransfer::ERHI::Invalid;
 			break;
 		}
 
diff --git a/Engine/Plugins/Media/AjaMedia/Source/AjaMediaOutput/Private/AjaMediaOutputModule.cpp b/Engine/Plugins/Media/AjaMedia/Source/AjaMediaOutput/Private/AjaMediaOutputModule.cpp
--- a/Engine/Plugins/Media/AjaMedia/Source/AjaMediaOutput/Private/AjaMediaOutputModule.cpp
+++ b/Engine/Plugins/Media/AjaMedia/Source/AjaMediaOutput/Private/AjaMediaOutputModule.cpp
@@ -44,43 +44,38 @@ void FAjaMediaOutputModule::StartupModule()
 
 		if (bIsGPUTextureTransferAvailable)
 		{
-			ENQUEUE_RENDER_COMMAND(AjaMediaCaptureInitialize)(
-				[this](FRHICommandListImmediate& RHICmdList) mutable
+			if (!GDynamicRHI)
+			{
+				return;
+			}
+
+			auto GetRHI = []()
+			{
+				switch (RHIGetInterfaceType())
 				{
-					if (!GDynamicRHI)
-					{
-						return;
-					}
-
-					auto GetRHI = []()
-					{
-						switch (RHIGetInterfaceType())
-						{
-						case ERHIInterfaceType::D3D11: return AJA::ERHI::D3D11;
-						case ERHIInterfaceType::D3D12: return AJA::ERHI::D3D12;
-						case ERHIInterfaceType::Vulkan: return AJA::ERHI::Vulkan;
-						default: return AJA::ERHI::Invalid;
-						}
-					};
-
-					AJA::FInitializeDMAArgs Args;
-					AJA::ERHI RHI = GetRHI();
-					Args.RHI = RHI;
-					/* Re-enable when adding vulkan support
-					if (RHI == AJA::ERHI::Vulkan)
-					{
-						FVulkanDynamicRHI* vkDynamicRHI = GetDynamicRHI<FVulkanDynamicRHI>();
-						Args.InVulkanInstance = vkDynamicRHI->GetInstance();
-
-						FMemory::Memcpy(Args.InRHIDeviceUUID, vkDynamicRHI->GetDevice()->GetDeviceIdProperties().deviceUUID, 16);
-					}
-					*/
-					Args.RHIDevice = GDynamicRHI->RHIGetNativeDevice();
-					Args.RHICommandQueue = GDynamicRHI->RHIGetNativeGraphicsQueue();
-
-					UE_LOG(LogAjaMediaOutput, Display, TEXT("Initializing GPU Texture transfer from AJA"));
-					bIsGPUTextureTransferAvailable = AJA::InitializeDMA(Args);
-				});
+				case ERHIInterfaceType::D3D12: return AJA::ERHI::D3D12;
+				case ERHIInterfaceType::Vulkan: return AJA::ERHI::Vulkan;
+				default: return AJA::ERHI::Invalid;
+				}
+			};
+
+			AJA::FInitializeDMAArgs Args;
+			AJA::ERHI RHI = GetRHI();
+			Args.RHI = RHI;
+			/* Re-enable when adding vulkan support
+			if (RHI == AJA::ERHI::Vulkan)
+			{
+				FVulkanDynamicRHI* vkDynamicRHI = GetDynamicRHI<FVulkanDynamicRHI>();
+				Args.InVulkanInstance = vkDynamicRHI->GetInstance();
+
+				FMemory::Memcpy(Args.InRHIDeviceUUID, vkDynamicRHI->GetDevice()->GetDeviceIdProperties().deviceUUID, 16);
+			}
+			*/
+			Args.RHIDevice = GDynamicRHI->RHIGetNativeDevice();
+			Args.RHICommandQueue = GDynamicRHI->RHIGetNativeGraphicsQueue();
+
+			UE_LOG(LogAjaMediaOutput, Display, TEXT("Initializing GPU Texture transfer from AJA"));
+			bIsGPUTextureTransferAvailable = AJA::InitializeDMA(Args);
 		}
 	};
 	
@@ -88,11 +83,7 @@ void FAjaMediaOutputModule::StartupModule()
 	{
 		if (bIsGPUTextureTransferAvailable)
 		{
-			ENQUEUE_RENDER_COMMAND(AjaMediaOutputModuleUninitialize)(
-				[](FRHICommandListImmediate& RHICmdList) mutable
-				{
-					AJA::UninitializeDMA();
-				});
+			AJA::UninitializeDMA();
 		}
 	};
 
diff --git a/Engine/Plugins/Media/BlackmagicMedia/Source/BlackmagicMediaOutput/Private/BlackmagicMediaOutputModule.cpp b/Engine/Plugins/Media/BlackmagicMedia/Source/BlackmagicMediaOutput/Private/BlackmagicMediaOutputModule.cpp
--- a/Engine/Plugins/Media/BlackmagicMedia/Source/BlackmagicMediaOutput/Private/BlackmagicMediaOutputModule.cpp
+++ b/Engine/Plugins/Media/BlackmagicMedia/Source/BlackmagicMediaOutput/Private/BlackmagicMediaOutputModule.cpp
@@ -42,35 +42,31 @@ void FBlackmagicMediaOutputModule::StartupModule()
 			}
 		}
 
+
 		if (bIsGPUTextureTransferAvailable)
 		{
-			ENQUEUE_RENDER_COMMAND(BlackmagicMediaCaptureInitialize)(
-				[this](FRHICommandListImmediate& RHICmdList) mutable
+			if (!GDynamicRHI)
+			{
+				return;
+			}
+
+			auto GetRHI = []()
+			{
+				switch (RHIGetInterfaceType())
 				{
-					if (!GDynamicRHI)
-					{
-						return;
-					}
-
-					auto GetRHI = []()
-					{
-						switch (RHIGetInterfaceType())
-						{
-						case ERHIInterfaceType::D3D11: return BlackmagicDesign::ERHI::D3D11;
-						case ERHIInterfaceType::D3D12: return BlackmagicDesign::ERHI::D3D12;
-						case ERHIInterfaceType::Vulkan: return BlackmagicDesign::ERHI::Vulkan;
-						default: return BlackmagicDesign::ERHI::Invalid;
-						}
-					};
-
-					BlackmagicDesign::FInitializeDMAArgs Args;
-					BlackmagicDesign::ERHI RHI = GetRHI();
-					Args.RHI = RHI;
-					Args.RHIDevice = GDynamicRHI->RHIGetNativeDevice();
-					Args.RHICommandQueue = GDynamicRHI->RHIGetNativeGraphicsQueue();
-
-					bIsGPUTextureTransferAvailable = BlackmagicDesign::InitializeDMA(Args);
-				});
+				case ERHIInterfaceType::D3D12: return BlackmagicDesign::ERHI::D3D12;
+				case ERHIInterfaceType::Vulkan: return BlackmagicDesign::ERHI::Vulkan;
+				default: return BlackmagicDesign::ERHI::Invalid;
+				}
+			};
+
+			BlackmagicDesign::FInitializeDMAArgs Args;
+			BlackmagicDesign::ERHI RHI = GetRHI();
+			Args.RHI = RHI;
+			Args.RHIDevice = GDynamicRHI->RHIGetNativeDevice();
+			Args.RHICommandQueue = GDynamicRHI->RHIGetNativeGraphicsQueue();
+
+			bIsGPUTextureTransferAvailable = BlackmagicDesign::InitializeDMA(Args);
 		}
 	};
 
@@ -79,11 +75,7 @@ void FBlackmagicMediaOutputModule::StartupModule()
 	{
 		if (bIsGPUTextureTransferAvailable)
 		{
-			ENQUEUE_RENDER_COMMAND(BlackmagicMediaCaptureUninitialize)(
-				[](FRHICommandListImmediate& RHICmdList) mutable
-				{
-					BlackmagicDesign::UninitializeDMA();
-				});
+			BlackmagicDesign::UninitializeDMA();
 		}
 	};
 
diff --git a/Engine/Plugins/Media/MediaIOFramework/Source/GPUTextureTransfer/Private/GPUTextureTransferModule.cpp b/Engine/Plugins/Media/MediaIOFramework/Source/GPUTextureTransfer/Private/GPUTextureTransferModule.cpp
--- a/Engine/Plugins/Media/MediaIOFramework/Source/GPUTextureTransfer/Private/GPUTextureTransferModule.cpp
+++ b/Engine/Plugins/Media/MediaIOFramework/Source/GPUTextureTransfer/Private/GPUTextureTransferModule.cpp
@@ -26,7 +26,6 @@
 #include "Misc/App.h"
 #include "Misc/CoreDelegates.h"
 #include "Modules/ModuleManager.h"
-#include "RenderingThread.h"
 
 #include <unordered_map>
 
@@ -170,71 +169,63 @@ void FGPUTextureTransferModule::InitializeTextureTransfer()
 		return;
 	}
 
-	ENQUEUE_RENDER_COMMAND(InitializeGPUTextureTransfer)(
-	[this](FRHICommandListImmediate& RHICmdList) mutable
+	if (!GDynamicRHI)
 	{
-		if (!GDynamicRHI)
-		{
-			return;
-		}
+		return;
+	}
 
-		UE::GPUTextureTransfer::TextureTransferPtr TextureTransfer;
+	UE::GPUTextureTransfer::TextureTransferPtr TextureTransfer;
 
-		UE::GPUTextureTransfer::ERHI RHI = ConvertRHI(RHIGetInterfaceType());
+	UE::GPUTextureTransfer::ERHI RHI = ConvertRHI(RHIGetInterfaceType());
 
-		switch (RHI)
-		{
-		case UE::GPUTextureTransfer::ERHI::D3D11:
-			TextureTransfer = MakeShared<UE::GPUTextureTransfer::Private::FD3D11TextureTransfer>();
-			break;
-		case UE::GPUTextureTransfer::ERHI::D3D12:
-			TextureTransfer = MakeShared<UE::GPUTextureTransfer::Private::FD3D12TextureTransfer>();
-			break;
-		case UE::GPUTextureTransfer::ERHI::Vulkan:
-			TextureTransfer = MakeShared<UE::GPUTextureTransfer::Private::FVulkanTextureTransfer>();
-			break;
-		default:
-			ensureAlways(false);
-			break;
-		}
+	switch (RHI)
+	{
+	case UE::GPUTextureTransfer::ERHI::D3D11:
+		TextureTransfer = MakeShared<UE::GPUTextureTransfer::Private::FD3D11TextureTransfer>();
+		break;
+	case UE::GPUTextureTransfer::ERHI::D3D12:
+		TextureTransfer = MakeShared<UE::GPUTextureTransfer::Private::FD3D12TextureTransfer>();
+		break;
+	case UE::GPUTextureTransfer::ERHI::Vulkan:
+		TextureTransfer = MakeShared<UE::GPUTextureTransfer::Private::FVulkanTextureTransfer>();
+		break;
+	default:
+		ensureAlways(false);
+		break;
+	}
 
-		UE::GPUTextureTransfer::FInitializeDMAArgs InitializeArgs;
-		InitializeArgs.RHI = RHI;
-		InitializeArgs.RHIDevice = GDynamicRHI->RHIGetNativeDevice();
-		InitializeArgs.RHICommandQueue = GDynamicRHI->RHIGetNativeGraphicsQueue();
+	UE::GPUTextureTransfer::FInitializeDMAArgs InitializeArgs;
+	InitializeArgs.RHI = RHI;
+	InitializeArgs.RHIDevice = GDynamicRHI->RHIGetNativeDevice();
+	InitializeArgs.RHICommandQueue = GDynamicRHI->RHIGetNativeGraphicsQueue();
 #if VULKAN_PLATFORM
-		if (RHI == UE::GPUTextureTransfer::ERHI::Vulkan)
-		{
-			IVulkanDynamicRHI* DynRHI = GetIVulkanDynamicRHI();
-			InitializeArgs.VulkanInstance = DynRHI->RHIGetVkInstance();
-			FMemory::Memcpy(InitializeArgs.RHIDeviceUUID, DynRHI->RHIGetVulkanDeviceUUID(), 16);
-		}
+	if (RHI == UE::GPUTextureTransfer::ERHI::Vulkan)
+	{
+		IVulkanDynamicRHI* DynRHI = GetIVulkanDynamicRHI();
+		InitializeArgs.VulkanInstance = DynRHI->RHIGetVkInstance();
+		FMemory::Memcpy(InitializeArgs.RHIDeviceUUID, DynRHI->RHIGetVulkanDeviceUUID(), 16);
+	}
 #endif
 
-		const uint8 RHIIndex = static_cast<uint8>(RHI);
-		UE_LOG(LogGPUTextureTransfer, Display, TEXT("Initializing GPU Texture transfer"));
-		if (TextureTransfer->Initialize(InitializeArgs))
-		{
-			TransferObjects[RHIIndex] = TextureTransfer;
-		}
-	});
+	const uint8 RHIIndex = static_cast<uint8>(RHI);
+	UE_LOG(LogGPUTextureTransfer, Display, TEXT("Initializing GPU Texture transfer"));
+	if (TextureTransfer->Initialize(InitializeArgs))
+	{
+		TransferObjects[RHIIndex] = TextureTransfer;
+	}
 #endif // DVP_SUPPORTED_PLATFORM
 }
 
 void FGPUTextureTransferModule::UninitializeTextureTransfer()
 {
 #if DVP_SUPPORTED_PLATFORM
-	ENQUEUE_RENDER_COMMAND(UninitializeGPUTextureTransfer)(
-		[this](FRHICommandListImmediate& RHICmdList) mutable
+	for (uint8 RhiIt = 1; RhiIt < RHI_MAX; RhiIt++)
+	{
+		if (const UE::GPUTextureTransfer::TextureTransferPtr& TextureTransfer = TransferObjects[RhiIt])
 		{
-			for (uint8 RhiIt = 1; RhiIt < RHI_MAX; RhiIt++)
-			{
-				if (const UE::GPUTextureTransfer::TextureTransferPtr& TextureTransfer = TransferObjects[RhiIt])
-				{
-					TextureTransfer->Uninitialize();
-				}
-			}
-		});
+			TextureTransfer->Uninitialize();
+		}
+	}
 #endif
 }
 
