Fix CIS warning related to media io audio output debug

diff --git a/Engine/Plugins/Media/MediaIOFramework/Source/MediaIOCore/MediaIOCore.Build.cs b/Engine/Plugins/Media/MediaIOFramework/Source/MediaIOCore/MediaIOCore.Build.cs
--- a/Engine/Plugins/Media/MediaIOFramework/Source/MediaIOCore/MediaIOCore.Build.cs
+++ b/Engine/Plugins/Media/MediaIOFramework/Source/MediaIOCore/MediaIOCore.Build.cs
@@ -48,7 +48,7 @@ namespace UnrealBuildTool.Rules
 	                "SoundFieldRendering"
 				});
 
-			PublicDefinitions.Add("WITH_MEDIA_IO_AUDIO_DEBUGGING=1");
+			PublicDefinitions.Add("WITH_MEDIA_IO_AUDIO_DEBUGGING=0");
 
 			if (Target.bBuildEditor == true)
 			{
