- Fixing media capture crop functionality

diff --git a/Engine/Plugins/Media/MediaIOFramework/Source/MediaIOCore/Private/MediaCapture.cpp b/Engine/Plugins/Media/MediaIOFramework/Source/MediaIOCore/Private/MediaCapture.cpp
--- a/Engine/Plugins/Media/MediaIOFramework/Source/MediaIOCore/Private/MediaCapture.cpp
+++ b/Engine/Plugins/Media/MediaIOFramework/Source/MediaIOCore/Private/MediaCapture.cpp
@@ -525,7 +525,7 @@ namespace UE::MediaCaptureData
 			// Default to no crop
 			OutSizeU = { 0.0f, 1.0f };
 			OutSizeV = { 0.0f, 1.0f };
-			OutCopyInfo.Size = FIntVector(SourceTexture->GetSizeX(), SourceTexture->GetSizeY(), 1);
+			OutCopyInfo.Size = FIntVector(Args.DesiredSize.X, Args.DesiredSize.Y, 1);
 			if (Args.MediaCapture->DesiredCaptureOptions.Crop != EMediaCaptureCroppingType::None)
 			{
 				switch (Args.MediaCapture->DesiredCaptureOptions.Crop)
