Disable gpu texture transfer if render doc plugin is activated

diff --git a/Engine/Plugins/Media/MediaIOFramework/Source/GPUTextureTransfer/Private/GPUTextureTransferModule.cpp b/Engine/Plugins/Media/MediaIOFramework/Source/GPUTextureTransfer/Private/GPUTextureTransferModule.cpp
--- a/Engine/Plugins/Media/MediaIOFramework/Source/GPUTextureTransfer/Private/GPUTextureTransferModule.cpp
+++ b/Engine/Plugins/Media/MediaIOFramework/Source/GPUTextureTransfer/Private/GPUTextureTransferModule.cpp
@@ -150,8 +150,16 @@ void FGPUTextureTransferModule::InitializeTextureTransfer()
 
 	// This must be called on game thread 
 	const FGPUDriverInfo GPUDriverInfo = FPlatformMisc::GetGPUDriverInfo(GRHIAdapterName);
-	bIsGPUTextureTransferAvailable = GPUDriverInfo.IsNVIDIA() && !FModuleManager::Get().IsModuleLoaded("RenderDocPlugin") && !GPUDriverInfo.DeviceDescription.Contains(TEXT("Tesla"));
+	bIsGPUTextureTransferAvailable = GPUDriverInfo.IsNVIDIA() && !GPUDriverInfo.DeviceDescription.Contains(TEXT("Tesla"));
 
+	const bool bRenderDocAttached = FParse::Param(FCommandLine::Get(), TEXT("AttachRenderDoc"));
+	if (bRenderDocAttached)
+	{
+		bIsGPUTextureTransferAvailable = false;
+		// Render doc clashes with GPU Direct.
+		UE_LOG(LogGPUTextureTransfer, Display, TEXT("GPU Texture Transfer disabled because render is attached."))
+	}
+	
 	if (bIsGPUTextureTransferAvailable)
 	{
 		bIsGPUTextureTransferAvailable = false;
