MediaIO: Fixed media capture immediate sometimes missing a frame when the game frame stalled

diff --git a/Engine/Plugins/Media/MediaIOFramework/Source/MediaIOCore/Private/MediaCapture.cpp b/Engine/Plugins/Media/MediaIOFramework/Source/MediaIOCore/Private/MediaCapture.cpp
--- a/Engine/Plugins/Media/MediaIOFramework/Source/MediaIOCore/Private/MediaCapture.cpp
+++ b/Engine/Plugins/Media/MediaIOFramework/Source/MediaIOCore/Private/MediaCapture.cpp
@@ -1284,7 +1284,7 @@ void UMediaCapture::CaptureImmediate_RenderThread(FRDGBuilder& GraphBuilder, FRH
 	FQueuedCaptureData NextCaptureData;
 	while (CaptureDataQueue.Dequeue(NextCaptureData))
 	{
-		if (NextCaptureData.BaseData.SourceFrameNumberRenderThread == GFrameNumberRenderThread)
+		if (NextCaptureData.BaseData.SourceFrameNumberRenderThread == GFrameCounterRenderThread)
 		{
 			bFoundMatchingData = true;
 			break;
@@ -1293,7 +1293,7 @@ void UMediaCapture::CaptureImmediate_RenderThread(FRDGBuilder& GraphBuilder, FRH
 
 	if (bFoundMatchingData == false)
 	{
-		UE_LOG(LogMediaIOCore, Warning, TEXT("Can't capture frame. Could not find a matching game frame."));
+		UE_LOG(LogMediaIOCore, Warning, TEXT("Can't capture frame. Could not find the matching game frame %d."), GFrameCounterRenderThread);
 		return;
 	}
 
@@ -1552,14 +1552,14 @@ bool UMediaCapture::CaptureRenderTargetImpl(UTextureRenderTarget2D * InRenderTar
 
 void UMediaCapture::OnBeginFrame_GameThread()
 {
-	
+
 	if (CaptureSourceType == EMediaCaptureSourceType::RHI_RESOURCE && (GetState() == EMediaCaptureState::Capturing))
 	{
 		// Queue capture data to be consumed when capture requests are done on render thread
 		FQueuedCaptureData CaptureData;
 		CaptureData.BaseData.SourceFrameTimecode = FApp::GetTimecode();
 		CaptureData.BaseData.SourceFrameTimecodeFramerate = FApp::GetTimecodeFrameRate();
-		CaptureData.BaseData.SourceFrameNumberRenderThread = GFrameNumber;
+		CaptureData.BaseData.SourceFrameNumberRenderThread = GFrameCounter;
 		CaptureData.UserData = GetCaptureFrameUserData_GameThread();
 		CaptureDataQueue.Enqueue(MoveTemp(CaptureData));
 	}
@@ -1637,7 +1637,7 @@ void UMediaCapture::OnEndFrame_GameThread()
 	{
 		CapturingFrame->CaptureBaseData.SourceFrameTimecode = FApp::GetTimecode();
 		CapturingFrame->CaptureBaseData.SourceFrameTimecodeFramerate = FApp::GetTimecodeFrameRate();
-		CapturingFrame->CaptureBaseData.SourceFrameNumberRenderThread = GFrameNumber;
+		CapturingFrame->CaptureBaseData.SourceFrameNumberRenderThread = GFrameCounter;
 		CapturingFrame->CaptureBaseData.SourceFrameNumber = ++CaptureRequestCount;
 		CapturingFrame->UserData = GetCaptureFrameUserData_GameThread();
 	}
