Refactor media capture
Allow passing rdg texture to immediate capture mode
Move frame manager into its own file
Use shared ptrs when possible for FCaptureFrame
Reduce duplication by introducing FCaptureSource, different implementations will handle the logic that differs for viewport capture and render targets

diff --git a/Engine/Plugins/Media/AjaMedia/Source/Aja/Private/OutputChannel.cpp b/Engine/Plugins/Media/AjaMedia/Source/Aja/Private/OutputChannel.cpp
--- a/Engine/Plugins/Media/AjaMedia/Source/Aja/Private/OutputChannel.cpp
+++ b/Engine/Plugins/Media/AjaMedia/Source/Aja/Private/OutputChannel.cpp
@@ -182,6 +182,8 @@ namespace AJA
 			, bInterlacedTest_ExpectFirstLineToBeWhite(false)
 			, InterlacedTest_FrameCounter(0)
 		{
+			TextureTransfer = FGPUTextureTransferModule::Get().GetTextureTransfer();
+
 			if (InOptions.OutputNumberOfBuffers > 0)
 			{
 				AllFrames.reserve(InOptions.OutputNumberOfBuffers);
diff --git a/Engine/Plugins/Media/AjaMedia/Source/Aja/Private/OutputChannel.h b/Engine/Plugins/Media/AjaMedia/Source/Aja/Private/OutputChannel.h
--- a/Engine/Plugins/Media/AjaMedia/Source/Aja/Private/OutputChannel.h
+++ b/Engine/Plugins/Media/AjaMedia/Source/Aja/Private/OutputChannel.h
@@ -3,6 +3,7 @@
 #pragma once
 
 #include "ChannelBase.h"
+#include "GPUTextureTransferModule.h"
 
 #include <vector>
 
@@ -137,6 +138,8 @@ namespace AJA
 			ULWord CurrentAudioWriteOffset = 0;
 			int32_t NumSamplesPerFrame = 0;
 			ULWord AudioPlayheadLastPosition = 0;
+
+			UE::GPUTextureTransfer::TextureTransferPtr TextureTransfer = nullptr;
 		};
 	}
 }
diff --git a/Engine/Plugins/Media/AjaMedia/Source/Aja/Private/TextureDMA.cpp b/Engine/Plugins/Media/AjaMedia/Source/Aja/Private/TextureDMA.cpp
deleted file mode 100644
diff --git a/Engine/Plugins/Media/AjaMedia/Source/Aja/Public/AJALib.h b/Engine/Plugins/Media/AjaMedia/Source/Aja/Public/AJALib.h
--- a/Engine/Plugins/Media/AjaMedia/Source/Aja/Public/AJALib.h
+++ b/Engine/Plugins/Media/AjaMedia/Source/Aja/Public/AJALib.h
@@ -18,8 +18,6 @@ namespace UE::GPUTextureTransfer
 
 namespace AJA
 {
-	extern UE::GPUTextureTransfer::ITextureTransfer* TextureTransfer;
-
 	/* Types provided from the interface
 	*****************************************************************************/
 	typedef void* FDeviceScanner;
@@ -535,49 +533,4 @@ namespace AJA
 	private:
 		Private::AutoDetectChannel* AutoChannel;
 	};
-
-	enum class ERHI : uint8_t
-	{
-		Invalid,
-		D3D11,
-		D3D12,
-		Vulkan
-	};
-
-	struct AJACORE_API FInitializeDMAArgs
-	{
-		ERHI RHI = ERHI::Invalid;
-		void* RHIDevice = nullptr;
-		void* RHICommandQueue = nullptr;
-
-		// Begin Vulkan Only
-		void* VulkanInstance = nullptr;
-		uint8_t RHIDeviceUUID[16] = { 0 };
-		// End Vulkan Only
-	};
-
-	struct AJACORE_API FRegisterDMABufferArgs
-	{
-		void* Buffer = nullptr;
-		uint32_t Stride = 0;
-		uint32_t Height = 0;
-		uint32_t Width = 0;
-		uint8_t BytesPerPixel = 0;
-	};
-
-	struct AJACORE_API FRegisterDMATextureArgs
-	{
-		FRHITexture* RHITexture = nullptr;
-		void* RHIResourceMemory = nullptr; // Vulkan only
-		uint32_t Stride = 0;
-	};
-
-	AJACORE_API bool InitializeDMA(const FInitializeDMAArgs& Args);
-	AJACORE_API void UninitializeDMA();
-	AJACORE_API bool RegisterDMATexture(const FRegisterDMATextureArgs& Args);
-	AJACORE_API bool UnregisterDMATexture(FRHITexture* RHITexture);
-	AJACORE_API bool RegisterDMABuffer(const FRegisterDMABufferArgs& InArgs);
-	AJACORE_API bool UnregisterDMABuffer(void* InBuffer);
-	AJACORE_API bool LockDMATexture(FRHITexture* RHITexture);
-	AJACORE_API bool UnlockDMATexture(FRHITexture* RHITexture);
 }
diff --git a/Engine/Plugins/Media/AjaMedia/Source/AjaMediaOutput/AjaMediaOutput.Build.cs b/Engine/Plugins/Media/AjaMedia/Source/AjaMediaOutput/AjaMediaOutput.Build.cs
--- a/Engine/Plugins/Media/AjaMedia/Source/AjaMediaOutput/AjaMediaOutput.Build.cs
+++ b/Engine/Plugins/Media/AjaMedia/Source/AjaMediaOutput/AjaMediaOutput.Build.cs
@@ -28,6 +28,7 @@ namespace UnrealBuildTool.Rules
 					"Core",
 					"CoreUObject",
 					"Engine",
+					"GPUTextureTransfer",
 					"MovieSceneCapture",
 					"RenderCore",
 					"RHI",
diff --git a/Engine/Plugins/Media/AjaMedia/Source/AjaMediaOutput/Private/AjaMediaCapture.cpp b/Engine/Plugins/Media/AjaMedia/Source/AjaMediaOutput/Private/AjaMediaCapture.cpp
--- a/Engine/Plugins/Media/AjaMedia/Source/AjaMediaOutput/Private/AjaMediaCapture.cpp
+++ b/Engine/Plugins/Media/AjaMedia/Source/AjaMediaOutput/Private/AjaMediaCapture.cpp
@@ -6,6 +6,7 @@
 #include "AjaDeviceProvider.h"
 #include "AjaMediaOutput.h"
 #include "AjaMediaOutputModule.h"
+#include "GPUTextureTransferModule.h"
 #include "Misc/CoreDelegates.h"
 #include "Engine/Engine.h"
 #include "HAL/Event.h"
@@ -224,7 +225,8 @@
 	, FrameRate(30, 1)
 	, WakeUpEvent(nullptr)
 {
-	bGPUTextureTransferAvailable = ((FAjaMediaOutputModule*)&IAjaMediaOutputModule::Get())->IsGPUTextureTransferAvailable();
+	bGPUTextureTransferAvailable = FGPUTextureTransferModule::Get().IsAvailable();
+	TextureTransfer = FGPUTextureTransferModule::Get().GetTextureTransfer();
 
 #if WITH_EDITOR
 	if (GEditor)
@@ -301,7 +303,7 @@
 	{
 		{
 			// Prevent the rendering thread from copying while we are stopping the capture.
-			FScopeLock ScopeLock(&RenderThreadCriticalSection);
+			FScopeLock ScopeLock(&CopyingCriticalSection);
 
 			ENQUEUE_RENDER_COMMAND(AjaMediaCaptureInitialize)(
 			[this](FRHICommandListImmediate& RHICmdList) mutable
@@ -311,7 +313,7 @@
 				{
 					for (FTextureRHIRef& Texture : TexturesToRelease)
 					{
-						AJA::UnregisterDMATexture(Texture->GetTexture2D());
+						TextureTransfer->UnregisterTexture(Texture->GetTexture2D());
 					}
 
 					TexturesToRelease.Reset();
@@ -513,86 +515,18 @@
 {
 	TRACE_CPUPROFILER_EVENT_SCOPE(UAjaMediaCapture::OnFrameCaptured_RenderingThread);
 
-	// Prevent the rendering thread from copying while we are stopping the capture.
-	FScopeLock ScopeLock(&RenderThreadCriticalSection);
-	if (OutputChannel)
-	{
-		const AJA::FTimecode Timecode = AjaMediaCaptureDevice::ConvertToAJATimecode(InBaseData.SourceFrameTimecode, InBaseData.SourceFrameTimecodeFramerate.AsDecimal(), FrameRate.AsDecimal());
-		const AjaMediaCaptureDevice::FAjaMediaEncodeOptions EncodeOptions(Width, Height, PixelFormat, UseKey);
-		
-		if (bEncodeTimecodeInTexel)
-		{
-			const FMediaIOCoreEncodeTime EncodeTime(EncodeOptions.EncodePixelFormat, InBuffer, EncodeOptions.Stride, EncodeOptions.TimeEncodeWidth, Height);
-			EncodeTime.Render(Timecode.Hours, Timecode.Minutes, Timecode.Seconds, Timecode.Frames);
-		}
+	FMediaCaptureResourceData InResourceData;
+	InResourceData.Buffer = InBuffer;
+	InResourceData.Width = Width;
+	InResourceData.Height = Height;
+	InResourceData.BytesPerRow = BytesPerRow;
 
-		AJA::AJAOutputFrameBufferData FrameBuffer;
-		FrameBuffer.Timecode = Timecode;
-		FrameBuffer.FrameIdentifier = InBaseData.SourceFrameNumber;
-		FrameBuffer.bEvenFrame = GFrameCounterRenderThread % 2 == 0;
-
-		bool bSetVideoResult = false;
-		{
-			TRACE_CPUPROFILER_EVENT_SCOPE(UAjaMediaCapture::OnFrameCaptured_RenderingThread::SetVideo);
-			bSetVideoResult = OutputChannel->SetVideoFrameData(FrameBuffer, reinterpret_cast<uint8_t*>(InBuffer), EncodeOptions.Stride * Height);
-		}
-
-		// If the set video call fails, that means we probably didn't find an available frame to write to,
-		// so don't pop from the audio buffer since we would lose these samples in the SetAudioFrameData call.
-		if (bSetVideoResult)
-		{
-			OutputAudio_RenderingThread(FrameBuffer);
-		}
-
-		if (bAjaWriteInputRawDataCmdEnable)
-		{
-			MediaIOCoreFileWriter::WriteRawFile(EncodeOptions.OutputFilename, reinterpret_cast<uint8*>(InBuffer), EncodeOptions.Stride * Height);
-			bAjaWriteInputRawDataCmdEnable = false;
-		}
-
-		WaitForSync_RenderingThread();
-		
-	}
-	else if (GetState() != EMediaCaptureState::Stopped)
-	{
-		SetState(EMediaCaptureState::Error);
-	}
+	OnFrameCapturedInternal_AnyThread(InBaseData, InUserData, MoveTemp(InResourceData));
 }
 
 void UAjaMediaCapture::OnRHIResourceCaptured_RenderingThread(const FCaptureBaseData& InBaseData, TSharedPtr<FMediaCaptureUserData, ESPMode::ThreadSafe> InUserData, FTextureRHIRef InTexture)
 {
-	TRACE_CPUPROFILER_EVENT_SCOPE(UAjaMediaCapture::OnFrameCaptured_RenderingThread);
-	
-	// Prevent the rendering thread from copying while we are stopping the capture.
-	FScopeLock ScopeLock(&RenderThreadCriticalSection);
-	if (OutputChannel)
-	{
-		const AJA::FTimecode Timecode = AjaMediaCaptureDevice::ConvertToAJATimecode(InBaseData.SourceFrameTimecode, InBaseData.SourceFrameTimecodeFramerate.AsDecimal(), FrameRate.AsDecimal());
-		
-		AJA::AJAOutputFrameBufferData FrameBuffer;
-		FrameBuffer.Timecode = Timecode;
-		FrameBuffer.FrameIdentifier = InBaseData.SourceFrameNumber;
-
-		bool bSetVideoResult = false;
-		{
-			TRACE_CPUPROFILER_EVENT_SCOPE(UAjaMediaCapture::OnFrameCaptured_RenderingThread::SetVideo_GPUDirect);
-			bSetVideoResult = OutputChannel->SetVideoFrameData(FrameBuffer, InTexture->GetTexture2D());
-		}
-
-		// If the set video call fails, that means we probably didn't find an available frame to write to,
-		// so don't pop from the audio buffer since we would lose these samples in the SetAudioFrameData call.
-		if (bSetVideoResult)
-		{
-			OutputAudio_RenderingThread(FrameBuffer);
-		}
-		
-		WaitForSync_RenderingThread();
-		
-	}
-	else if (GetState() != EMediaCaptureState::Stopped)
-	{
-		SetState(EMediaCaptureState::Error);
-	}
+	OnRHIResourceCaptured_AnyThread(InBaseData, InUserData, InTexture);
 }
 
 void UAjaMediaCapture::LockDMATexture_RenderThread(FTextureRHIRef InTexture)
@@ -604,37 +538,61 @@
 			TexturesToRelease.Add(InTexture);
 
 			FRHITexture2D* Texture = InTexture->GetTexture2D();
-			AJA::FRegisterDMATextureArgs Args;
+			UE::GPUTextureTransfer::FRegisterDMATextureArgs Args;
 			Args.RHITexture = Texture;
-			//Args.InRHIResourceMemory = Texture->GetNativeResource(); todo: VulkanTexture->Surface->GetAllocationHandle for Vulkan
-
-			AJA::RegisterDMATexture(Args);
+			
+			Args.RHIResourceMemory = nullptr; // = Texture->GetNativeResource(); todo: VulkanTexture->Surface->GetAllocationHandle for Vulkan
+			Args.Width = Texture->GetDesc().GetSize().X;
+			Args.Height = Texture->GetDesc().GetSize().Y;
+			
+			if (Args.RHITexture->GetFormat() == EPixelFormat::PF_B8G8R8A8)
+			{
+				Args.PixelFormat = UE::GPUTextureTransfer::EPixelFormat::PF_8Bit;
+				Args.Stride = Args.Width * 4;
+			}
+			else if (Args.RHITexture->GetFormat() == EPixelFormat::PF_R32G32B32A32_UINT)
+			{
+				Args.PixelFormat = UE::GPUTextureTransfer::EPixelFormat::PF_10Bit;
+				Args.Stride = Args.Width * 16;
+			}
+			else
+			{
+				checkf(false, TEXT("Format not supported"));
+			}
+			
+			TextureTransfer->RegisterTexture(Args);
 		}
 	}
 
-	AJA::LockDMATexture(InTexture->GetTexture2D());
+	TextureTransfer->LockTexture(InTexture->GetTexture2D());
 }
 
 void UAjaMediaCapture::UnlockDMATexture_RenderThread(FTextureRHIRef InTexture)
 {
 	TRACE_CPUPROFILER_EVENT_SCOPE(UAjaMediaCapture::OnFrameCaptured_RenderingThread::UnlockDMATexture);
-	AJA::UnlockDMATexture(InTexture->GetTexture2D());
+	TextureTransfer->UnlockTexture(InTexture->GetTexture2D());
 }
 
-void UAjaMediaCapture::WaitForSync_RenderingThread() const
+void UAjaMediaCapture::OnFrameCaptured_AnyThread(const FCaptureBaseData& InBaseData, TSharedPtr<FMediaCaptureUserData, ESPMode::ThreadSafe> InUserData, const FMediaCaptureResourceData& InResourceData)
 {
+	TRACE_CPUPROFILER_EVENT_SCOPE(UAjaMediaCapture::OnFrameCaptured_AnyThread);
+	OnFrameCapturedInternal_AnyThread(InBaseData, InUserData, InResourceData);
+}
+
+void UAjaMediaCapture::WaitForSync_AnyThread() const
+{
 	if (bWaitForSyncEvent)
 	{
-		if (WakeUpEvent && GetState() != EMediaCaptureState::Error) // In render thread, could be shutdown in a middle of a frame
+		if (WakeUpEvent && GetState() != EMediaCaptureState::Error) // Could be shutdown in a middle of a frame 
 		{
 			WakeUpEvent->Wait();
 		}
 	}
 }
 
-void UAjaMediaCapture::OutputAudio_RenderingThread(const AJA::AJAOutputFrameBufferData& FrameBuffer) const
+void UAjaMediaCapture::OutputAudio_AnyThread(const AJA::AJAOutputFrameBufferData& FrameBuffer) const
 {
-	if (bOutputAudio)
+	if (bOutputAudio && AudioOutput)
 	{
 		TRACE_CPUPROFILER_EVENT_SCOPE(UAjaMediaCapture::OnFrameCaptured_RenderingThread::SetAudio);
 
@@ -719,3 +677,88 @@
 	check(false);
 	return false;
 }
+
+void UAjaMediaCapture::OnRHIResourceCaptured_AnyThread(const FCaptureBaseData& InBaseData, TSharedPtr<FMediaCaptureUserData, ESPMode::ThreadSafe> InUserData, FTextureRHIRef InTexture)
+{
+	TRACE_CPUPROFILER_EVENT_SCOPE(UAjaMediaCapture::OnFrameCaptured_AnyThread);
+	
+	// Prevent the rendering thread from copying while we are stopping the capture.
+	FScopeLock ScopeLock(&CopyingCriticalSection);
+	if (OutputChannel)
+	{
+		const AJA::FTimecode Timecode = AjaMediaCaptureDevice::ConvertToAJATimecode(InBaseData.SourceFrameTimecode, InBaseData.SourceFrameTimecodeFramerate.AsDecimal(), FrameRate.AsDecimal());
+		
+		AJA::AJAOutputFrameBufferData FrameBuffer;
+		FrameBuffer.Timecode = Timecode;
+		FrameBuffer.FrameIdentifier = InBaseData.SourceFrameNumber;
+
+		bool bSetVideoResult = false;
+		{
+			TRACE_CPUPROFILER_EVENT_SCOPE(UAjaMediaCapture::OnFrameCaptured_RenderingThread::SetVideo_GPUDirect);
+			bSetVideoResult = OutputChannel->SetVideoFrameData(FrameBuffer, InTexture->GetTexture2D());
+		}
+
+		// If the set video call fails, that means we probably didn't find an available frame to write to,
+		// so don't pop from the audio buffer since we would lose these samples in the SetAudioFrameData call.
+		if (bSetVideoResult)
+		{
+			OutputAudio_AnyThread(FrameBuffer);
+		}
+		
+		WaitForSync_AnyThread();
+		
+	}
+	else if (GetState() != EMediaCaptureState::Stopped)
+	{
+		SetState(EMediaCaptureState::Error);
+	}
+}
+
+void UAjaMediaCapture::OnFrameCapturedInternal_AnyThread(const FCaptureBaseData& InBaseData, TSharedPtr<FMediaCaptureUserData, ESPMode::ThreadSafe> InUserData, const FMediaCaptureResourceData& InResourceData)
+{
+	// Prevent this thread from copying while we are stopping the capture.
+	FScopeLock ScopeLock(&CopyingCriticalSection);
+	if (OutputChannel)
+	{
+		const AJA::FTimecode Timecode = AjaMediaCaptureDevice::ConvertToAJATimecode(InBaseData.SourceFrameTimecode, InBaseData.SourceFrameTimecodeFramerate.AsDecimal(), FrameRate.AsDecimal());
+		const AjaMediaCaptureDevice::FAjaMediaEncodeOptions EncodeOptions(InResourceData.Width, InResourceData.Height, PixelFormat, UseKey);
+
+		if (bEncodeTimecodeInTexel)
+		{
+			const FMediaIOCoreEncodeTime EncodeTime(EncodeOptions.EncodePixelFormat, InResourceData.Buffer, EncodeOptions.Stride, EncodeOptions.TimeEncodeWidth, InResourceData.Height);
+			EncodeTime.Render(Timecode.Hours, Timecode.Minutes, Timecode.Seconds, Timecode.Frames);
+		}
+
+		AJA::AJAOutputFrameBufferData FrameBuffer;
+		FrameBuffer.Timecode = Timecode;
+		FrameBuffer.FrameIdentifier = InBaseData.SourceFrameNumber;
+
+		// This will most likely be wrong when using the AnyThread callback 
+		FrameBuffer.bEvenFrame = GFrameCounterRenderThread % 2 == 0;
+
+		bool bSetVideoResult = false;
+		{
+			TRACE_CPUPROFILER_EVENT_SCOPE(UAjaMediaCapture::OnFrameCaptured::SetVideo);
+			bSetVideoResult = OutputChannel->SetVideoFrameData(FrameBuffer, reinterpret_cast<uint8_t*>(InResourceData.Buffer), EncodeOptions.Stride * InResourceData.Height);
+		}
+
+		// If the set video call fails, that means we probably didn't find an available frame to write to,
+		// so don't pop from the audio buffer since we would lose these samples in the SetAudioFrameData call.
+		if (bSetVideoResult)
+		{
+			OutputAudio_AnyThread(FrameBuffer);
+		}
+
+		if (bAjaWriteInputRawDataCmdEnable)
+		{
+			MediaIOCoreFileWriter::WriteRawFile(EncodeOptions.OutputFilename, reinterpret_cast<uint8*>(InResourceData.Buffer), EncodeOptions.Stride * InResourceData.Height);
+			bAjaWriteInputRawDataCmdEnable = false;
+		}
+
+		WaitForSync_AnyThread();
+	}
+	else if (GetState() != EMediaCaptureState::Stopped)
+	{
+		SetState(EMediaCaptureState::Error);
+	}
+}
diff --git a/Engine/Plugins/Media/AjaMedia/Source/AjaMediaOutput/Private/AjaMediaOutputModule.cpp b/Engine/Plugins/Media/AjaMedia/Source/AjaMediaOutput/Private/AjaMediaOutputModule.cpp
--- a/Engine/Plugins/Media/AjaMedia/Source/AjaMediaOutput/Private/AjaMediaOutputModule.cpp
+++ b/Engine/Plugins/Media/AjaMedia/Source/AjaMediaOutput/Private/AjaMediaOutputModule.cpp
@@ -2,105 +2,16 @@
 
 #include "AjaMediaOutputModule.h"
 
-#include "AjaMediaFrameGrabberProtocol.h"
-#include "AjaLib.h"
-#include "DynamicRHI.h"
-#include "GenericPlatform/GenericPlatformDriver.h"
-#include "Misc/CoreDelegates.h"
-#include "RenderingThread.h"
-#include "RHI.h"
-
 #define LOCTEXT_NAMESPACE "AjaMediaOutput"
 
 DEFINE_LOG_CATEGORY(LogAjaMediaOutput);
 
 void FAjaMediaOutputModule::StartupModule()
 {
-	const auto DMAInitializationFunc = [this]()
-	{
-		static const TArray<FString> SupportedGPUPrefixes= { 
-			TEXT("RTX A4"),
-			TEXT("RTX A5"),
-			TEXT("RTX A6"),
-			TEXT("Quadro")
-		};
-
-		const FGPUDriverInfo GPUDriverInfo = FPlatformMisc::GetGPUDriverInfo(GRHIAdapterName);
-		bIsGPUTextureTransferAvailable = GPUDriverInfo.IsNVIDIA() &&!FModuleManager::Get().IsModuleLoaded("RenderDocPlugin");
-		bIsGPUTextureTransferAvailable &= !GPUDriverInfo.DeviceDescription.Contains(TEXT("Tesla"));
-		
-		if (bIsGPUTextureTransferAvailable)
-		{
-			bIsGPUTextureTransferAvailable = false;
-			for (const FString& GPUPrefix : SupportedGPUPrefixes)
-			{
-				if (GPUDriverInfo.DeviceDescription.Contains(GPUPrefix))
-				{
-					bIsGPUTextureTransferAvailable = true;
-					break;
-				}
-			}
-		}
-
-		if (bIsGPUTextureTransferAvailable)
-		{
-			if (!GDynamicRHI)
-			{
-				return;
-			}
-
-			auto GetRHI = []()
-			{
-				switch (RHIGetInterfaceType())
-				{
-				case ERHIInterfaceType::D3D12: return AJA::ERHI::D3D12;
-				case ERHIInterfaceType::Vulkan: return AJA::ERHI::Vulkan;
-				default: return AJA::ERHI::Invalid;
-				}
-			};
-
-			AJA::FInitializeDMAArgs Args;
-			AJA::ERHI RHI = GetRHI();
-			Args.RHI = RHI;
-			/* Re-enable when adding vulkan support
-			if (RHI == AJA::ERHI::Vulkan)
-			{
-				FVulkanDynamicRHI* vkDynamicRHI = GetDynamicRHI<FVulkanDynamicRHI>();
-				Args.InVulkanInstance = vkDynamicRHI->GetInstance();
-
-				FMemory::Memcpy(Args.InRHIDeviceUUID, vkDynamicRHI->GetDevice()->GetDeviceIdProperties().deviceUUID, 16);
-			}
-			*/
-			Args.RHIDevice = GDynamicRHI->RHIGetNativeDevice();
-			Args.RHICommandQueue = GDynamicRHI->RHIGetNativeGraphicsQueue();
-
-			UE_LOG(LogAjaMediaOutput, Display, TEXT("Initializing GPU Texture transfer from AJA"));
-			bIsGPUTextureTransferAvailable = AJA::InitializeDMA(Args);
-		}
-	};
-	
-	auto DMAUninitialize = [this]()
-	{
-		if (bIsGPUTextureTransferAvailable)
-		{
-			AJA::UninitializeDMA();
-		}
-	};
-
-	//Postpone initialization after all modules have been loaded to be sure Aja library has been loaded
-	FCoreDelegates::OnAllModuleLoadingPhasesComplete.AddLambda(DMAInitializationFunc);
-	//Same for shutdown, uninitialize ourselves before library is unloaded
-	FCoreDelegates::OnEnginePreExit.AddLambda(DMAUninitialize);
 }
 
 void FAjaMediaOutputModule::ShutdownModule()
 {
-
-}
-
-bool FAjaMediaOutputModule::IsGPUTextureTransferAvailable() const
-{
-	return bIsGPUTextureTransferAvailable;
 }
 
 IMPLEMENT_MODULE(FAjaMediaOutputModule, AjaMediaOutput )
diff --git a/Engine/Plugins/Media/AjaMedia/Source/AjaMediaOutput/Private/AjaMediaOutputModule.h b/Engine/Plugins/Media/AjaMedia/Source/AjaMediaOutput/Private/AjaMediaOutputModule.h
--- a/Engine/Plugins/Media/AjaMedia/Source/AjaMediaOutput/Private/AjaMediaOutputModule.h
+++ b/Engine/Plugins/Media/AjaMedia/Source/AjaMediaOutput/Private/AjaMediaOutputModule.h
@@ -10,9 +10,4 @@ class FAjaMediaOutputModule : public IAjaMediaOutputModule
 public:
 	virtual void StartupModule() override;
 	virtual void ShutdownModule() override;
-
-	bool IsGPUTextureTransferAvailable() const;
-
-private:
-	bool bIsGPUTextureTransferAvailable = false;
 };
\ No newline at end of file
diff --git a/Engine/Plugins/Media/AjaMedia/Source/AjaMediaOutput/Public/AjaMediaCapture.h b/Engine/Plugins/Media/AjaMedia/Source/AjaMediaOutput/Public/AjaMediaCapture.h
--- a/Engine/Plugins/Media/AjaMedia/Source/AjaMediaOutput/Public/AjaMediaCapture.h
+++ b/Engine/Plugins/Media/AjaMedia/Source/AjaMediaOutput/Public/AjaMediaCapture.h
@@ -20,6 +20,11 @@ namespace AJA
 	struct AJAOutputFrameBufferData;
 }
 
+namespace UE::GPUTextureTransfer
+{
+	class ITextureTransfer;
+}
+
 /**
  * Output Media for AJA streams.
  * The output format could be any of EAjaMediaOutputPixelFormat.
@@ -46,8 +51,11 @@ protected:
 	
 	virtual void OnFrameCaptured_RenderingThread(const FCaptureBaseData& InBaseData, TSharedPtr<FMediaCaptureUserData, ESPMode::ThreadSafe> InUserData, void* InBuffer, int32 Width, int32 Height, int32 BytesPerRow) override;
 	virtual void OnRHIResourceCaptured_RenderingThread(const FCaptureBaseData& InBaseData, TSharedPtr<FMediaCaptureUserData, ESPMode::ThreadSafe> InUserData, FTextureRHIRef InTexture) override;
+	virtual void OnRHIResourceCaptured_AnyThread(const FCaptureBaseData& InBaseData, TSharedPtr<FMediaCaptureUserData, ESPMode::ThreadSafe> InUserData, FTextureRHIRef InTexture) override;
 	virtual void LockDMATexture_RenderThread(FTextureRHIRef InTexture) override;
 	virtual void UnlockDMATexture_RenderThread(FTextureRHIRef InTexture) override;
+	virtual void OnFrameCaptured_AnyThread(const FCaptureBaseData& InBaseData, TSharedPtr<FMediaCaptureUserData, ESPMode::ThreadSafe> InUserData, const FMediaCaptureResourceData& InResourceData) override;
+
 
 private:
 	struct FAjaOutputCallback;
@@ -56,12 +64,13 @@ private:
 
 private:
 	bool InitAJA(UAjaMediaOutput* InMediaOutput);
-	void WaitForSync_RenderingThread() const;
-	void OutputAudio_RenderingThread(const AJA::AJAOutputFrameBufferData& FrameBuffer) const;
+	void WaitForSync_AnyThread() const;
+	void OutputAudio_AnyThread(const AJA::AJAOutputFrameBufferData& FrameBuffer) const;
 	void ApplyViewportTextureAlpha(TSharedPtr<FSceneViewport> InSceneViewport);
 	void RestoreViewportTextureAlpha(TSharedPtr<FSceneViewport> InSceneViewport);
 	bool CleanupPreEditorExit();
 	void OnEnginePreExit();
+	void OnFrameCapturedInternal_AnyThread(const FCaptureBaseData& InBaseData, TSharedPtr<FMediaCaptureUserData, ESPMode::ThreadSafe> InUserData, const FMediaCaptureResourceData& InResourceData);
 
 	struct FAudioBuffer
 	{
@@ -92,7 +101,7 @@ private:
 	FFrameRate FrameRate;
 
 	/** Critical section for synchronizing access to the OutputChannel */
-	FCriticalSection RenderThreadCriticalSection;
+	FCriticalSection CopyingCriticalSection;
 
 	/** Event to wakeup When waiting for sync */
 	FEvent* WakeUpEvent;
@@ -112,4 +121,6 @@ private:
 	FDelegateHandle CanCloseEditorDelegateHandle;
 
 	mutable int32 AudioSamplesSentLastFrame = 0;
+
+	TSharedPtr<UE::GPUTextureTransfer::ITextureTransfer> TextureTransfer;
 };
diff --git a/Engine/Plugins/Media/BlackmagicMedia/Source/BlackmagicMediaOutput/Private/BlackmagicMediaCapture.cpp b/Engine/Plugins/Media/BlackmagicMedia/Source/BlackmagicMediaOutput/Private/BlackmagicMediaCapture.cpp
--- a/Engine/Plugins/Media/BlackmagicMedia/Source/BlackmagicMediaOutput/Private/BlackmagicMediaCapture.cpp
+++ b/Engine/Plugins/Media/BlackmagicMedia/Source/BlackmagicMediaOutput/Private/BlackmagicMediaCapture.cpp
@@ -320,9 +320,8 @@ namespace BlackmagicMediaCaptureAnalytics
 
 ///* UBlackmagicMediaCapture implementation
 //*****************************************************************************/
-UBlackmagicMediaCapture::UBlackmagicMediaCapture(const FObjectInitializer& ObjectInitializer)
-	: Super(ObjectInitializer)
-	, bWaitForSyncEvent(false)
+UBlackmagicMediaCapture::UBlackmagicMediaCapture()
+	: bWaitForSyncEvent(false)
 	, bEncodeTimecodeInTexel(false)
 	, bLogDropFrame(false)
 	, BlackmagicMediaOutputPixelFormat(EBlackmagicMediaOutputPixelFormat::PF_8BIT_YUV)
diff --git a/Engine/Plugins/Media/BlackmagicMedia/Source/BlackmagicMediaOutput/Public/BlackmagicMediaCapture.h b/Engine/Plugins/Media/BlackmagicMedia/Source/BlackmagicMediaOutput/Public/BlackmagicMediaCapture.h
--- a/Engine/Plugins/Media/BlackmagicMedia/Source/BlackmagicMediaOutput/Public/BlackmagicMediaCapture.h
+++ b/Engine/Plugins/Media/BlackmagicMedia/Source/BlackmagicMediaOutput/Public/BlackmagicMediaCapture.h
@@ -29,8 +29,10 @@ namespace BlackmagicDesign
 UCLASS(BlueprintType)
 class BLACKMAGICMEDIAOUTPUT_API UBlackmagicMediaCapture : public UMediaCapture
 {
-	GENERATED_UCLASS_BODY()
+	GENERATED_BODY()
 
+public:
+	UBlackmagicMediaCapture();
 	//~ UMediaCapture interface
 public:
 	virtual bool HasFinishedProcessing() const override;
diff --git a/Engine/Plugins/Media/MediaIOFramework/Source/MediaIOCore/MediaIOCore.Build.cs b/Engine/Plugins/Media/MediaIOFramework/Source/MediaIOCore/MediaIOCore.Build.cs
--- a/Engine/Plugins/Media/MediaIOFramework/Source/MediaIOCore/MediaIOCore.Build.cs
+++ b/Engine/Plugins/Media/MediaIOFramework/Source/MediaIOCore/MediaIOCore.Build.cs
@@ -8,6 +8,7 @@ namespace UnrealBuildTool.Rules
 	{
 		public MediaIOCore(ReadOnlyTargetRules Target) : base(Target)
 		{
+
 			PublicDependencyModuleNames.AddRange(
 				new string[]
 				{
@@ -29,11 +30,12 @@ namespace UnrealBuildTool.Rules
 				});
 
 			var EngineDir = Path.GetFullPath(Target.RelativeEnginePath);
+
 			PrivateIncludePaths.AddRange(
 				new string[] {
-					//required for FScreenPassVS and AddDrawScreenPass
-					Path.Combine(EngineDir, "Source", "Runtime", "Renderer", "Private")
-				});
+				//required for FScreenPassVS, AddDrawScreenPass, and for scene view extensions related headers
+				Path.Combine(EngineDir, "Source", "Runtime", "Renderer", "Private")
+			});
 
 			PrivateDependencyModuleNames.AddRange(
                 new string[]
diff --git a/Engine/Plugins/Media/MediaIOFramework/Source/MediaIOCore/Private/MediaCapture.cpp b/Engine/Plugins/Media/MediaIOFramework/Source/MediaIOCore/Private/MediaCapture.cpp
--- a/Engine/Plugins/Media/MediaIOFramework/Source/MediaIOCore/Private/MediaCapture.cpp
+++ b/Engine/Plugins/Media/MediaIOFramework/Source/MediaIOCore/Private/MediaCapture.cpp
@@ -2,36 +2,34 @@
 
 #include "MediaCapture.h"
 
-
 #include "Application/ThrottleManager.h"
 #include "Async/Async.h"
 #include "Engine/GameEngine.h"
 #include "Engine/RendererSettings.h"
 #include "Engine/TextureRenderTarget2D.h"
 #include "EngineModule.h"
+#include "HAL/ThreadManager.h"
+#include "MediaCaptureSceneViewExtension.h"
+#include "MediaCaptureSources.h"
+#include "MediaIOFrameManager.h"
 #include "MediaIOCoreModule.h"
 #include "MediaOutput.h"
 #include "MediaShaders.h"
 #include "Misc/App.h"
 #include "Misc/ScopeLock.h"
-#include "PipelineStateCache.h"
 #include "RendererInterface.h"
 #include "RenderGraphResources.h"
 #include "RenderGraphUtils.h"
-#include "RenderUtils.h"
-#include "RenderTargetPool.h"
-#include "RHIStaticStates.h"
-#include "SceneUtils.h"
 #include "ScreenPass.h"
 #include "Slate/SceneViewport.h"
 #include "TextureResource.h"
+#include "ProfilingDebugging/ScopedTimers.h"
 #include "UObject/WeakObjectPtrTemplates.h"
 
 #include UE_INLINE_GENERATED_CPP_BY_NAME(MediaCapture)
 
 #if WITH_EDITOR
 #include "AnalyticsEventAttribute.h"
-#include "Editor.h"
 #include "Editor/EditorEngine.h"
 #include "Engine/Engine.h"
 #include "Engine/EngineTypes.h"
@@ -44,56 +42,28 @@
 
 #define LOCTEXT_NAMESPACE "MediaCapture"
 
+static TAutoConsoleVariable<int32> CVarMediaIOEnableExperimentalScheduling(
+	TEXT("MediaIO.EnableExperimentalScheduling"), 1,
+	TEXT("Whether to send out frame  in a separate thread. (Experimental)"),
+	ECVF_RenderThreadSafe);
+
+static TAutoConsoleVariable<int32> CVarMediaIOScheduleOnAnyThread(
+	TEXT("MediaIO.ScheduleOnAnyThread"), 0,
+	TEXT("Whether to wait for resource readback in a separate thread. (Experimental)"),
+	ECVF_RenderThreadSafe);
 
 /** Time spent in media capture sending a frame. */
 DECLARE_CYCLE_STAT(TEXT("MediaCapture RenderThread FrameCapture"), STAT_MediaCapture_RenderThread_FrameCapture, STATGROUP_Media);
 DECLARE_CYCLE_STAT(TEXT("MediaCapture RenderThread LockResource"), STAT_MediaCapture_RenderThread_LockResource, STATGROUP_Media);
-DECLARE_CYCLE_STAT(TEXT("MediaCapture RenderThread RHI Capture Callback"), STAT_MediaCapture_RenderThread_CaptureCallback, STATGROUP_Media);
+DECLARE_CYCLE_STAT(TEXT("MediaCapture AnyThread LockResource"), STAT_MediaCapture_AnyThread_LockResource, STATGROUP_Media);
+DECLARE_CYCLE_STAT(TEXT("MediaCapture RenderThread RHI Capture Callback"), STAT_MediaCapture_RenderThread_RHI_CaptureCallback, STATGROUP_Media);
+DECLARE_CYCLE_STAT(TEXT("MediaCapture RenderThread CPU Capture Callback"), STAT_MediaCapture_RenderThread_CaptureCallback, STATGROUP_Media);
 DECLARE_GPU_STAT(MediaCapture_CaptureFrame);
 DECLARE_GPU_STAT(MediaCapture_CustomCapture);
 DECLARE_GPU_STAT(MediaCapture_Conversion);
 DECLARE_GPU_STAT(MediaCapture_Readback);
 DECLARE_GPU_STAT(MediaCapture_ProcessCapture);
-
-
-/** These pixel formats do not require additional conversion except for swizzling and normalized sampling. */
-static TSet<EPixelFormat> SupportedRgbaSwizzleFormats =
-{
-	PF_A32B32G32R32F,
-	PF_B8G8R8A8,
-	PF_G8,
-	PF_G16,
-	PF_FloatRGB,
-	PF_FloatRGBA,
-	PF_R32_FLOAT,
-	PF_G16R16,
-	PF_G16R16F,
-	PF_G32R32F,
-	PF_A2B10G10R10,
-	PF_A16B16G16R16,
-	PF_R16F,
-	PF_FloatR11G11B10,
-	PF_A8,
-	PF_R32_UINT,
-	PF_R32_SINT,
-	PF_R16_UINT,
-	PF_R16_SINT,
-	PF_R16G16B16A16_UINT,
-	PF_R16G16B16A16_SINT,
-	PF_R5G6B5_UNORM,
-	PF_R8G8B8A8,
-	PF_A8R8G8B8,
-	PF_R8G8,
-	PF_R32G32B32A32_UINT,
-	PF_R16G16_UINT,
-	PF_R8_UINT,
-	PF_R8G8B8A8_UINT,
-	PF_R8G8B8A8_SNORM,
-	PF_R16G16B16A16_UNORM,
-	PF_R16G16B16A16_SNORM,
-	PF_R32G32_UINT,
-	PF_R8,
-};
+DECLARE_GPU_STAT(MediaCapture_SyncPointPass);
 
 /* namespace MediaCaptureDetails definition
 *****************************************************************************/
@@ -102,12 +72,8 @@ namespace MediaCaptureDetails
 {
 	bool FindSceneViewportAndLevel(TSharedPtr<FSceneViewport>& OutSceneViewport);
 
-	//Validation for the source of a capture
-	bool ValidateSceneViewport(const TSharedPtr<FSceneViewport>& SceneViewport, const FMediaCaptureOptions& CaptureOption, const FIntPoint& DesiredSize, const EPixelFormat DesiredPixelFormat, const bool bCurrentlyCapturing);
-	bool ValidateTextureRenderTarget2D(const UTextureRenderTarget2D* RenderTarget, const FMediaCaptureOptions& CaptureOption, const FIntPoint& DesiredSize, const EPixelFormat DesiredPixelFormat, const bool bCurrentlyCapturing);
-
 	//Validation that there is a capture
-	bool ValidateIsCapturing(const UMediaCapture& CaptureToBeValidated);
+	bool ValidateIsCapturing(const UMediaCapture* CaptureToBeValidated);
 
 	void ShowSlateNotification();
 
@@ -140,7 +106,7 @@ namespace MediaCaptureAnalytics
 
 namespace UE::MediaCaptureData
 {
-	class FCaptureFrame
+	class FCaptureFrame : public FFrame
 	{
 	public:
 		FCaptureFrame(int32 InFrameId)
@@ -148,12 +114,22 @@ namespace UE::MediaCaptureData
 			, bReadbackRequested(false)
 			, bDoingGPUCopy(false)
 		{
+		}
 
+		virtual int32 GetId() const override
+		{
+			return FrameId;
 		}
-		virtual ~FCaptureFrame() {};
 
+		virtual bool IsPending() const override
+		{
+			return bReadbackRequested || bDoingGPUCopy; 
+		}
+		
 		/** Returns true if its output resource is valid */
 		virtual bool HasValidResource() const = 0;
+		
+		virtual FRDGViewableResource* RegisterResource(FRDGBuilder& RDGBuilder) = 0;
 
 		/** Simple way to validate the resource type and cast safely */
 		virtual bool IsTextureResource() const = 0;
@@ -162,16 +138,21 @@ namespace UE::MediaCaptureData
 		/** Locks the readback resource and returns a pointer to access data from system memory */
 		virtual void* Lock(FRHICommandListImmediate& RHICmdList, FRHIGPUMask GPUMask, int32& OutRowStride) = 0;
 
+		virtual void* Lock_Unsafe(FRHIGPUMask GPUMask, int32& OutRowStride) { return nullptr; }
+
 		/** Unlocks the readback resource */
 		virtual void Unlock() = 0;
 
+		virtual void Unlock_Unsafe(FRHIGPUMask GPUMask) {}
+
 		/** Returns true if the readback is ready to be used */
 		virtual bool IsReadbackReady(FRHIGPUMask GPUMask) = 0;
+		
+		virtual void EnqueueCopy(FRDGBuilder& RDGBuilder, FRDGViewableResource* ResourceToReadback) const = 0;
 
 		virtual FRHITexture* GetTextureResource() { return nullptr; }
 		virtual FRHIBuffer* GetBufferResource() { return nullptr; }
-
-		bool IsPending() const { return bReadbackRequested || bDoingGPUCopy; }
+		virtual FRHIResource* GetRHIResource() = 0;
 
 		int32 FrameId = 0;
 		UMediaCapture::FCaptureBaseData CaptureBaseData;
@@ -180,11 +161,21 @@ namespace UE::MediaCaptureData
 		TSharedPtr<FMediaCaptureUserData> UserData;
 	};
 
+	/** Parameter to make our sync pass needing the convert pass as a prereq */
+	BEGIN_SHADER_PARAMETER_STRUCT(FMediaCaptureTextureSyncPassParameters, )
+		RDG_TEXTURE_ACCESS(Resource, ERHIAccess::CopySrc)
+	END_SHADER_PARAMETER_STRUCT()
+
+	BEGIN_SHADER_PARAMETER_STRUCT(FMediaCaptureBufferSyncPassParameters, )
+		RDG_BUFFER_ACCESS(Resource, ERHIAccess::CopySrc)
+	END_SHADER_PARAMETER_STRUCT()
+	
 	class FTextureCaptureFrame : public FCaptureFrame
 	{
 	public:
 		/** Type alias for the output resource type used during capture frame */
 		using FOutputResourceType = FRDGTextureRef;
+		using PassParametersType = FMediaCaptureTextureSyncPassParameters;
 
 		FTextureCaptureFrame(int32 InFrameId)
 			: FCaptureFrame(InFrameId)
@@ -211,7 +202,7 @@ namespace UE::MediaCaptureData
 		{
 			if (ReadbackTexture->IsReady(GPUMask) == false)
 			{
-				UE_LOG(LogMediaIOCore, Verbose, TEXT("Fenced for texture readback was not ready"));
+				UE_LOG(LogMediaIOCore, Verbose, TEXT("Fence for texture readback was not ready"));
 			}
 
 			int32 ReadbackWidth;
@@ -220,11 +211,25 @@ namespace UE::MediaCaptureData
 			return ReadbackPointer;
 		}
 
+		virtual void* Lock_Unsafe(FRHIGPUMask GPUMask, int32& OutRowStride) override
+		{
+			void* ReadbackPointer = nullptr;
+			int32 ReadbackWidth, ReadbackHeight;
+			GDynamicRHI->RHIMapStagingSurface(ReadbackTexture->DestinationStagingTextures[GPUMask.GetFirstIndex()], nullptr, ReadbackPointer, ReadbackWidth, ReadbackHeight, GPUMask.GetFirstIndex());
+			OutRowStride = ReadbackWidth * MediaCaptureDetails::GetBytesPerPixel(RenderTarget->GetDesc().Format);
+			return ReadbackPointer;
+		}
+
 		virtual void Unlock() override
 		{
 			ReadbackTexture->Unlock();
 		}
 
+		virtual void Unlock_Unsafe(FRHIGPUMask GPUMask) override
+		{
+			GDynamicRHI->RHIUnmapStagingSurface(ReadbackTexture->DestinationStagingTextures[0], GPUMask.GetFirstIndex());
+		}
+
 		virtual bool IsReadbackReady(FRHIGPUMask GPUMask) override
 		{
 			return ReadbackTexture->IsReady(GPUMask);
@@ -233,24 +238,24 @@ namespace UE::MediaCaptureData
 		//~ End FCaptureFrame interface
 
 		/** Registers an external texture to be tracked by the graph and returns a pointer to the tracked resource */
-		FRDGTextureRef RegisterResource(FRDGBuilder& RDGBuilder)
+		virtual FRDGViewableResource* RegisterResource(FRDGBuilder& RDGBuilder) override
 		{
 			return RDGBuilder.RegisterExternalTexture(RenderTarget, TEXT("OutputTexture"));
 		}
 
 		/** Adds a readback pass to the graph */
-		void EnqueueCopy(FRDGBuilder& RDGBuilder, FRDGTextureRef ResourceToReadback)
+		virtual void EnqueueCopy(FRDGBuilder& RDGBuilder, FRDGViewableResource* ResourceToReadback) const override
 		{
-			AddEnqueueCopyPass(RDGBuilder, ReadbackTexture.Get(), ResourceToReadback);
+			AddEnqueueCopyPass(RDGBuilder, ReadbackTexture.Get(), static_cast<FRDGTexture*>(ResourceToReadback));
 		}
 
 		virtual FRHITexture* GetTextureResource() override
 		{
-			return GetRHIResource();
+			return static_cast<FRHITexture*>(GetRHIResource());
 		}
 
 		/** Returns RHI resource of the allocated pooled resource */
-		FRHITexture* GetRHIResource()
+		virtual FRHIResource* GetRHIResource() override
 		{
 			return RenderTarget->GetRHI();
 		}
@@ -265,7 +270,8 @@ namespace UE::MediaCaptureData
 	public:
 		/** Type alias for the output resource type used during capture frame */
 		using FOutputResourceType = FRDGBufferRef;
-
+		using PassParametersType = FMediaCaptureBufferSyncPassParameters;
+		
 		FBufferCaptureFrame(int32 InFrameId)
 			: FCaptureFrame(InFrameId)
 		{
@@ -311,24 +317,24 @@ namespace UE::MediaCaptureData
 		//~ End FCaptureFrame interface
 
 		/** Registers an external texture to be tracked by the graph and returns a pointer to the tracked resource */
-		FRDGBufferRef RegisterResource(FRDGBuilder& RDGBuilder)
+		virtual FRDGViewableResource* RegisterResource(FRDGBuilder& RDGBuilder) override
 		{
 			return RDGBuilder.RegisterExternalBuffer(Buffer, TEXT("OutputBuffer"));
 		}
 
 		/** Adds a readback pass to the graph */
-		void EnqueueCopy(FRDGBuilder& RDGBuilder, FRDGBufferRef ResourceToReadback)
+		virtual void EnqueueCopy(FRDGBuilder& RDGBuilder, FRDGViewableResource* ResourceToReadback) const
 		{
-			AddEnqueueCopyPass(RDGBuilder, ReadbackBuffer.Get(), ResourceToReadback, Buffer->GetRHI()->GetSize());
+			AddEnqueueCopyPass(RDGBuilder, ReadbackBuffer.Get(), static_cast<FRDGBuffer*>(ResourceToReadback), Buffer->GetRHI()->GetSize());
 		}
 
 		virtual FRHIBuffer* GetBufferResource() override
 		{
-			return GetRHIResource();
+			return static_cast<FRHIBuffer*>(GetRHIResource());
 		}
 
 		/** Returns RHI resource of the allocated pooled resource */
-		FRHIBuffer* GetRHIResource()
+		virtual FRHIResource* GetRHIResource() override
 		{
 			return Buffer->GetRHI();
 		}
@@ -337,91 +343,48 @@ namespace UE::MediaCaptureData
 		TRefCountPtr<FRDGPooledBuffer> Buffer;
 		TUniquePtr<FRHIGPUBufferReadback> ReadbackBuffer;
 	};
-	
-	class FFrameManager
-	{
-	public:
-
-		~FFrameManager()
-		{
-			ENQUEUE_RENDER_COMMAND(MediaCaptureFrameManagerCleaning)(
-				[FramesToBeReleased = MoveTemp(CaptureFrames)](FRHICommandListImmediate& RHICmdList) mutable
-			{
-				FramesToBeReleased.Reset();
-			});
-		}
 
-		TConstArrayView<TUniquePtr<FCaptureFrame>> GetFrames() const
-		{
-			return CaptureFrames;
-		}
-
-		void AddFrame(TUniquePtr<FCaptureFrame>&& NewFrame)
-		{
-			AvailableFrames.Enqueue(NewFrame->FrameId);
-			CaptureFrames.Add(MoveTemp(NewFrame));
-		}
+	/** Helper struct to contain arguments for CaptureFrame */
+	struct FCaptureFrameArgs
+	{
+		FRDGBuilder& GraphBuilder;
+		TObjectPtr<UMediaCapture> MediaCapture = nullptr;
+		FTexture2DRHIRef ResourceToCapture;
+		FRDGTextureRef RDGResourceToCapture = nullptr;
+		FIntPoint DesiredSize = FIntPoint::ZeroValue;
 
-		void MarkAvailable(const FCaptureFrame& InFrame)
+		EPixelFormat GetFormat() const
 		{
-			if (ensure(!InFrame.IsPending()))
+			if (RDGResourceToCapture)
 			{
-				AvailableFrames.Enqueue(InFrame.FrameId);
+				return RDGResourceToCapture->Desc.Format;
 			}
+			return ResourceToCapture->GetFormat();
 		}
 
-		void MarkPending(const FCaptureFrame& InFrame)
+		bool HasValidResource() const
 		{
-			if (ensure(InFrame.IsPending()))
-			{
-				PendingFrames.Enqueue(InFrame.FrameId);
-			}
+			return ResourceToCapture || RDGResourceToCapture;
 		}
 
-		void CompleteNextPending(const FCaptureFrame& InFrame)
+		FIntPoint GetSizeXY() const
 		{
-			TOptional<int32> NextPending = PendingFrames.Dequeue();
-			if (ensure(NextPending.IsSet() && InFrame.FrameId == NextPending.GetValue()))
+			if (RDGResourceToCapture)
 			{
-				MarkAvailable(InFrame);
+				return FIntPoint(RDGResourceToCapture->Desc.GetSize().X, RDGResourceToCapture->Desc.GetSize().Y);
 			}
+			return ResourceToCapture->GetSizeXY();
 		}
 
-		FCaptureFrame* GetNextAvailable()
+		int32 GetSizeX() const
 		{
-			TOptional<int32> NextAvailable = AvailableFrames.Dequeue();
-			if (NextAvailable.IsSet())
-			{
-				return CaptureFrames[NextAvailable.GetValue()].Get();
-			}
-
-			return nullptr;
+			return GetSizeXY().X;
 		}
 
-		FCaptureFrame* PeekNextPending()
+		int32 GetSizeY() const
 		{
-			int32* NextPending = PendingFrames.Peek();
-			if (NextPending)
-			{
-				return CaptureFrames[*NextPending].Get();
-			}
-
-			return nullptr;
+			return GetSizeXY().Y;
 		}
-
-	private:
-		TArray<TUniquePtr<FCaptureFrame>> CaptureFrames;
-		TSpscQueue<int32> AvailableFrames;
-		TSpscQueue<int32> PendingFrames;
-	};
-
-	/** Helper struct to contain arguments for CaptureFrame */
-	struct FCaptureFrameArgs
-	{
-		FRDGBuilder& GraphBuilder;
-		TObjectPtr<UMediaCapture> MediaCapture = nullptr;
-		FTexture2DRHIRef ResourceToCapture;
-		FIntPoint DesiredSize = FIntPoint::ZeroValue;
 	};
 
 	/** Helper struct to contain arguments for AddConversionPass */
@@ -439,61 +402,30 @@ namespace UE::MediaCaptureData
 	{
 	public:
 
-		static FTexture2DRHIRef GetSourceTextureForInput(FRHICommandListImmediate& RHICmdList, FSceneViewport* CapturingSceneViewport, FTextureRenderTargetResource* TextureRenderTargetResource)
-		{
-			FTexture2DRHIRef SourceTexture;
-			if (CapturingSceneViewport)
-			{
-#if WITH_EDITOR
-				if (!IsRunningGame())
-				{
-					// PIE, PIE in windows, editor viewport
-					SourceTexture = CapturingSceneViewport->GetRenderTargetTexture();
-					if (!SourceTexture.IsValid() && CapturingSceneViewport->GetViewportRHI())
-					{
-						SourceTexture = RHICmdList.GetViewportBackBuffer(CapturingSceneViewport->GetViewportRHI());
-					}
-				}
-				else
-#endif
-					if (CapturingSceneViewport->GetViewportRHI())
-					{
-						// Standalone and packaged
-						SourceTexture = RHICmdList.GetViewportBackBuffer(CapturingSceneViewport->GetViewportRHI());
-					}
-			}
-			else if (TextureRenderTargetResource && TextureRenderTargetResource->GetTextureRenderTarget2DResource())
-			{
-				SourceTexture = TextureRenderTargetResource->GetTextureRenderTarget2DResource()->GetTextureRHI();
-			}
-
-			return SourceTexture;
-		}
-
 		static bool AreInputsValid(const FCaptureFrameArgs& Args)
 		{
 			// If it is a simple rgba swizzle we can handle the conversion. Supported formats
 			// contained in SupportedRgbaSwizzleFormats. Warning would've been displayed on start of capture.
-			if (Args.MediaCapture->DesiredPixelFormat != Args.ResourceToCapture->GetFormat() &&
-				(!SupportedRgbaSwizzleFormats.Contains(Args.ResourceToCapture->GetFormat()) || !Args.MediaCapture->DesiredCaptureOptions.bConvertToDesiredPixelFormat))
+			if (Args.MediaCapture->DesiredPixelFormat != Args.GetFormat() &&
+				(!UMediaCapture::GetSupportedRgbaSwizzleFormats().Contains(Args.GetFormat()) || !Args.MediaCapture->DesiredCaptureOptions.bConvertToDesiredPixelFormat))
 			{
 				UE_LOG(LogMediaIOCore, Error, TEXT("The capture will stop for '%s'. The Source pixel format doesn't match with the user requested pixel format. %sRequested: %s Source: %s")
 					, *Args.MediaCapture->MediaOutputName
-					, (SupportedRgbaSwizzleFormats.Contains(Args.ResourceToCapture->GetFormat()) && !Args.MediaCapture->DesiredCaptureOptions.bConvertToDesiredPixelFormat) ? TEXT("Please enable \"Convert To Desired Pixel Format\" option in Media Capture settings. ") : TEXT("")
+					, (UMediaCapture::GetSupportedRgbaSwizzleFormats().Contains(Args.GetFormat()) && !Args.MediaCapture->DesiredCaptureOptions.bConvertToDesiredPixelFormat) ? TEXT("Please enable \"Convert To Desired Pixel Format\" option in Media Capture settings. ") : TEXT("")
 					, GetPixelFormatString(Args.MediaCapture->DesiredPixelFormat)
-					, GetPixelFormatString(Args.ResourceToCapture->GetFormat()));
+					, GetPixelFormatString(Args.GetFormat()));
 
 				return false;
 			}
 
 			if (Args.MediaCapture->DesiredCaptureOptions.Crop == EMediaCaptureCroppingType::None)
 			{
-				if (Args.DesiredSize.X != Args.ResourceToCapture->GetSizeX() || Args.DesiredSize.Y != Args.ResourceToCapture->GetSizeY())
+				if (Args.DesiredSize.X != Args.GetSizeX() || Args.DesiredSize.Y != Args.GetSizeY())
 				{
 					UE_LOG(LogMediaIOCore, Error, TEXT("The capture will stop for '%s'. The Source size doesn't match with the user requested size. Requested: %d,%d  Source: %d,%d")
 						, *Args.MediaCapture->MediaOutputName
 						, Args.DesiredSize.X, Args.DesiredSize.Y
-						, Args.ResourceToCapture->GetSizeX(), Args.ResourceToCapture->GetSizeY());
+						, Args.GetSizeX(), Args.GetSizeY());
 
 					return false;
 				}
@@ -506,12 +438,12 @@ namespace UE::MediaCaptureData
 					StartCapturePoint = Args.MediaCapture->DesiredCaptureOptions.CustomCapturePoint;
 				}
 
-				if ((uint32)(Args.DesiredSize.X + StartCapturePoint.X) > Args.ResourceToCapture->GetSizeX() || (uint32)(Args.DesiredSize.Y + StartCapturePoint.Y) > Args.ResourceToCapture->GetSizeY())
+				if ((Args.DesiredSize.X + StartCapturePoint.X) > Args.GetSizeX() || (Args.DesiredSize.Y + StartCapturePoint.Y) > Args.GetSizeY())
 				{
 					UE_LOG(LogMediaIOCore, Error, TEXT("The capture will stop for '%s'. The Source size doesn't match with the user requested size. Requested: %d,%d  Source: %d,%d")
 						, *Args.MediaCapture->MediaOutputName
 						, Args.DesiredSize.X, Args.DesiredSize.Y
-						, Args.ResourceToCapture->GetSizeX(), Args.ResourceToCapture->GetSizeY());
+						, Args.GetSizeX(), Args.GetSizeY());
 
 					return false;
 				}
@@ -520,7 +452,7 @@ namespace UE::MediaCaptureData
 			return true;
 		}
 
-		static void GetCopyInfo(const FCaptureFrameArgs& Args, const FTexture2DRHIRef& SourceTexture, FRHICopyTextureInfo& OutCopyInfo, FVector2D& OutSizeU, FVector2D& OutSizeV)
+		static void GetCopyInfo(const FCaptureFrameArgs& Args, FRHICopyTextureInfo& OutCopyInfo, FVector2D& OutSizeU, FVector2D& OutSizeV)
 		{
 			// Default to no crop
 			OutSizeU = { 0.0f, 1.0f };
@@ -531,24 +463,34 @@ namespace UE::MediaCaptureData
 				switch (Args.MediaCapture->DesiredCaptureOptions.Crop)
 				{
 				case EMediaCaptureCroppingType::Center:
-					OutCopyInfo.SourcePosition = FIntVector((SourceTexture->GetSizeX() - Args.DesiredSize.X) / 2, (SourceTexture->GetSizeY() - Args.DesiredSize.Y) / 2, 0);
+					OutCopyInfo.SourcePosition = FIntVector((Args.GetSizeX() - Args.DesiredSize.X) / 2, (Args.GetSizeY() - Args.DesiredSize.Y) / 2, 0);
 					break;
 				case EMediaCaptureCroppingType::TopLeft:
 					break;
 				case EMediaCaptureCroppingType::Custom:
 					OutCopyInfo.SourcePosition = FIntVector(Args.MediaCapture->DesiredCaptureOptions.CustomCapturePoint.X, Args.MediaCapture->DesiredCaptureOptions.CustomCapturePoint.Y, 0);
 					break;
+				default:
+					break;
 				}
 
-				OutSizeU.X = (float)(OutCopyInfo.SourcePosition.X)                      / (float)SourceTexture->GetSizeX();
-				OutSizeU.Y = (float)(OutCopyInfo.SourcePosition.X + OutCopyInfo.Size.X) / (float)SourceTexture->GetSizeX();
-				OutSizeV.X = (float)(OutCopyInfo.SourcePosition.Y)                      / (float)SourceTexture->GetSizeY();
-				OutSizeV.Y = (float)(OutCopyInfo.SourcePosition.Y + OutCopyInfo.Size.Y) / (float)SourceTexture->GetSizeY();
+				OutSizeU.X = (float)(OutCopyInfo.SourcePosition.X)                      / (float)Args.GetSizeX();
+				OutSizeU.Y = (float)(OutCopyInfo.SourcePosition.X + OutCopyInfo.Size.X) / (float)Args.GetSizeX();
+				OutSizeV.X = (float)(OutCopyInfo.SourcePosition.Y)                      / (float)Args.GetSizeY();
+				OutSizeV.Y = (float)(OutCopyInfo.SourcePosition.Y + OutCopyInfo.Size.Y) / (float)Args.GetSizeY();
 			}
 		}
 
-		static void AddConversionPass(const FCaptureFrameArgs& Args, const FConversionPassArgs& ConversionPassArgs, FRDGTextureRef OutputResource)
+		static void AddConversionPass(const FCaptureFrameArgs& Args, const FConversionPassArgs& ConversionPassArgs, FRDGViewableResource* OutputResource)
 		{
+			if (OutputResource->Type == FRDGBuffer::StaticType)
+			{
+				return;
+			}
+
+			check(OutputResource->Type == FRDGTexture::StaticType);
+			FRDGTexture* OutputTexture =  static_cast<FRDGTexture*>(OutputResource);
+			
 			//Based on conversion type, this might be changed
 			bool bRequiresFormatConversion = ConversionPassArgs.bRequiresFormatConversion;
 
@@ -570,7 +512,7 @@ namespace UE::MediaCaptureData
 			// If no conversion was required, go through a simple copy
 			if (Args.MediaCapture->ConversionOperation == EMediaCaptureConversionOperation::NONE && !bRequiresFormatConversion)
 			{
-				AddCopyTexturePass(Args.GraphBuilder, ConversionPassArgs.SourceRGBTexture, OutputResource, ConversionPassArgs.CopyInfo);
+				AddCopyTexturePass(Args.GraphBuilder, ConversionPassArgs.SourceRGBTexture, OutputTexture, ConversionPassArgs.CopyInfo);
 			}
 			else
 			{
@@ -586,9 +528,9 @@ namespace UE::MediaCaptureData
 				{
 					//Configure source/output viewport to get the right UV scaling from source texture to output texture
 					FScreenPassTextureViewport InputViewport(ConversionPassArgs.SourceRGBTexture, ViewRect);
-					FScreenPassTextureViewport OutputViewport(OutputResource);
+					FScreenPassTextureViewport OutputViewport(OutputTexture);
 					TShaderMapRef<FRGB8toUYVY8ConvertPS> PixelShader(GlobalShaderMap);
-					FRGB8toUYVY8ConvertPS::FParameters* Parameters = PixelShader->AllocateAndSetParameters(Args.GraphBuilder, ConversionPassArgs.SourceRGBTexture, MediaShaders::RgbToYuvRec709Scaled, MediaShaders::YUVOffset8bits, bDoLinearToSRGB, OutputResource);
+					FRGB8toUYVY8ConvertPS::FParameters* Parameters = PixelShader->AllocateAndSetParameters(Args.GraphBuilder, ConversionPassArgs.SourceRGBTexture, MediaShaders::RgbToYuvRec709Scaled, MediaShaders::YUVOffset8bits, bDoLinearToSRGB, OutputTexture);
 					AddDrawScreenPass(Args.GraphBuilder, RDG_EVENT_NAME("RGBToUYVY 8 bit"), ViewInfo, OutputViewport, InputViewport, VertexShader, PixelShader, Parameters);
 				}
 				break;
@@ -597,9 +539,9 @@ namespace UE::MediaCaptureData
 					//Configure source/output viewport to get the right UV scaling from source texture to output texture
 					const FIntPoint InExtent = FIntPoint((((ConversionPassArgs.SourceRGBTexture->Desc.Extent.X + 47) / 48) * 48), ConversionPassArgs.SourceRGBTexture->Desc.Extent.Y);;
 					FScreenPassTextureViewport InputViewport(ConversionPassArgs.SourceRGBTexture, ViewRect);
-					FScreenPassTextureViewport OutputViewport(OutputResource);
+					FScreenPassTextureViewport OutputViewport(OutputTexture);
 					TShaderMapRef<FRGB10toYUVv210ConvertPS> PixelShader(GlobalShaderMap);
-					FRGB10toYUVv210ConvertPS::FParameters* Parameters = PixelShader->AllocateAndSetParameters(Args.GraphBuilder, ConversionPassArgs.SourceRGBTexture, MediaShaders::RgbToYuvRec709Scaled, MediaShaders::YUVOffset10bits, bDoLinearToSRGB, OutputResource);
+					FRGB10toYUVv210ConvertPS::FParameters* Parameters = PixelShader->AllocateAndSetParameters(Args.GraphBuilder, ConversionPassArgs.SourceRGBTexture, MediaShaders::RgbToYuvRec709Scaled, MediaShaders::YUVOffset10bits, bDoLinearToSRGB, OutputTexture);
 					AddDrawScreenPass(Args.GraphBuilder, RDG_EVENT_NAME("RGBToYUVv210"), ViewInfo, OutputViewport, InputViewport, VertexShader, PixelShader, Parameters);
 				}
 				break;
@@ -614,7 +556,7 @@ namespace UE::MediaCaptureData
 					{
 						//Configure source/output viewport to get the right UV scaling from source texture to output texture
 						FScreenPassTextureViewport InputViewport(ConversionPassArgs.SourceRGBTexture, ViewRect);
-						FScreenPassTextureViewport OutputViewport(OutputResource);
+						FScreenPassTextureViewport OutputViewport(OutputTexture);
 
 						// In cases where texture is converted from a format that doesn't have A channel, we want to force set it to 1.
 						EMediaCaptureConversionOperation MediaConversionOperation = Args.MediaCapture->DesiredCaptureOptions.bForceAlphaToOneOnConversion ? EMediaCaptureConversionOperation::SET_ALPHA_ONE : Args.MediaCapture->ConversionOperation;
@@ -622,7 +564,7 @@ namespace UE::MediaCaptureData
 						PermutationVector.Set<FModifyAlphaSwizzleRgbaPS::FConversionOp>(static_cast<int32>(MediaConversionOperation));
 
 						TShaderMapRef<FModifyAlphaSwizzleRgbaPS> PixelShader(GlobalShaderMap, PermutationVector);
-						FModifyAlphaSwizzleRgbaPS::FParameters* Parameters = PixelShader->AllocateAndSetParameters(Args.GraphBuilder, ConversionPassArgs.SourceRGBTexture, OutputResource);
+						FModifyAlphaSwizzleRgbaPS::FParameters* Parameters = PixelShader->AllocateAndSetParameters(Args.GraphBuilder, ConversionPassArgs.SourceRGBTexture, OutputTexture);
 						AddDrawScreenPass(Args.GraphBuilder, RDG_EVENT_NAME("MediaCaptureSwizzle"), ViewInfo, OutputViewport, InputViewport, VertexShader, PixelShader, Parameters);
 					}
 					break;
@@ -630,13 +572,208 @@ namespace UE::MediaCaptureData
 			}
 		}
 
-		static void AddConversionPass(const FCaptureFrameArgs& Args, const FConversionPassArgs& ConversionPassArgs, FRDGBufferRef OutputResource)
+		template <typename CaptureType>
+		static void OnReadbackComplete(FRHICommandList& RHICmdList, UMediaCapture* MediaCapture, TSharedPtr<CaptureType> ReadyFrame)
+		{
+			UE_LOG(LogMediaIOCore, Verbose, TEXT("[%s - %s] - Processing pending frame %d"), *MediaCapture->MediaOutputName, *FThreadManager::GetThreadName(FPlatformTLS::GetCurrentThreadId()), ReadyFrame->GetId());
+			
+			{
+				ON_SCOPE_EXIT
+				{
+					if (IsInRenderingThread())
+					{
+						--MediaCapture->WaitingForRenderCommandExecutionCounter;
+					}
+				};
+
+				FRHIGPUMask GPUMask;
+	#if WITH_MGPU
+				GPUMask = RHICmdList.GetGPUMask();
+
+				// If GPUMask is not set to a specific GPU we and since we are reading back the texture, it shouldn't matter which GPU we do this on.
+				if (!GPUMask.HasSingleIndex())
+				{
+					GPUMask = FRHIGPUMask::FromIndex(GPUMask.GetFirstIndex());
+				}
+
+				SCOPED_GPU_MASK(RHICmdList, GPUMask);
+	#endif
+
+				// Are we doing a GPU Direct transfer
+				if (MediaCapture->ShouldCaptureRHIResource())
+				{
+					if (ReadyFrame->IsTextureResource())
+					{
+						{
+							TRACE_CPUPROFILER_EVENT_SCOPE(UMediaCapture::UnlockDMATexture_RenderThread);
+							MediaCapture->UnlockDMATexture_RenderThread(ReadyFrame->GetTextureResource());
+						}
+
+						TRACE_CPUPROFILER_EVENT_SCOPE(UMediaCapture::RHIResourceCaptured);
+						TRACE_CPUPROFILER_EVENT_SCOPE_TEXT(*FString::Printf(TEXT("MediaCapture Output Frame %d"), ReadyFrame->CaptureBaseData.SourceFrameNumberRenderThread));
+						SCOPE_CYCLE_COUNTER(STAT_MediaCapture_RenderThread_CaptureCallback)
+
+						if (IsInRenderingThread())
+						{
+							MediaCapture->OnRHIResourceCaptured_RenderingThread(ReadyFrame->CaptureBaseData, ReadyFrame->UserData, ReadyFrame->GetTextureResource());
+						}
+						else
+						{
+							MediaCapture->OnRHIResourceCaptured_AnyThread(ReadyFrame->CaptureBaseData, ReadyFrame->UserData, ReadyFrame->GetTextureResource());
+						}
+					}
+					else
+					{
+						TRACE_CPUPROFILER_EVENT_SCOPE(UMediaCapture::RHIResourceCaptured);
+						if (IsInRenderingThread())
+						{
+							MediaCapture->OnRHIResourceCaptured_RenderingThread(ReadyFrame->CaptureBaseData, ReadyFrame->UserData, ReadyFrame->GetBufferResource());
+						}
+						else
+						{
+							MediaCapture->OnRHIResourceCaptured_AnyThread(ReadyFrame->CaptureBaseData, ReadyFrame->UserData, ReadyFrame->GetBufferResource());
+						}
+					}
+
+					ReadyFrame->bDoingGPUCopy = false;
+				}
+				else
+				{
+					// Lock & read
+					void* ColorDataBuffer = nullptr;
+					int32 RowStride = 0;
+
+					// Readback should be ready since we're after the sync point.
+					ensure(ReadyFrame->IsReadbackReady(GPUMask));
+					{
+						SCOPE_CYCLE_COUNTER(STAT_MediaCapture_RenderThread_LockResource);
+						ColorDataBuffer = ReadyFrame->Lock_Unsafe(GPUMask, RowStride);
+					}
+
+					if (ColorDataBuffer)
+					{
+						{
+							TRACE_CPUPROFILER_EVENT_SCOPE_TEXT(*FString::Printf(TEXT("MediaCapture Output Frame %d"), ReadyFrame->CaptureBaseData.SourceFrameNumberRenderThread));
+							SCOPE_CYCLE_COUNTER(STAT_MediaCapture_AnyThread_LockResource)
+
+							if (IsInRenderingThread())
+							{
+								MediaCapture->OnFrameCaptured_RenderingThread(ReadyFrame->CaptureBaseData, ReadyFrame->UserData, ColorDataBuffer, MediaCapture->DesiredOutputSize.X, MediaCapture->DesiredOutputSize.Y, RowStride);
+							}
+							else
+							{
+								UMediaCapture::FMediaCaptureResourceData ResourceData;
+								ResourceData.Buffer = ColorDataBuffer;
+								ResourceData.Width = MediaCapture->DesiredOutputSize.X;
+								ResourceData.Height = MediaCapture->DesiredOutputSize.Y;
+								ResourceData.BytesPerRow = RowStride;
+								MediaCapture->OnFrameCaptured_AnyThread(ReadyFrame->CaptureBaseData, ReadyFrame->UserData, MoveTemp(ResourceData));
+							}
+
+						}
+						ReadyFrame->Unlock_Unsafe(GPUMask);
+					}
+
+					ReadyFrame->bReadbackRequested = false;
+				}
+
+				UE_LOG(LogMediaIOCore, Verbose, TEXT("[%s - %s] - Completed pending frame %d."), *MediaCapture->MediaOutputName, *FThreadManager::GetThreadName(FPlatformTLS::GetCurrentThreadId()), ReadyFrame->GetId());
+				MediaCapture->FrameManager->CompleteNextPending(*ReadyFrame);
+				--MediaCapture->PendingFrameCount;
+			};
+		}
+
+
+		template <typename CaptureType>
+		static void AddSyncPointPass(FRDGBuilder& GraphBuilder, const TObjectPtr<UMediaCapture>& MediaCapture, TSharedPtr<FCaptureFrame> CapturingFrame, FRDGViewableResource* OutputResource)
 		{
-			/** Used at some point to deal with native buffer output resource type but there to validate compiled template */
+			TRACE_CPUPROFILER_EVENT_SCOPE(MediaCaptureSyncPoint);
+			TRACE_CPUPROFILER_EVENT_SCOPE_TEXT(*FString::Printf(TEXT("MediaCaptureSyncPoint_%d"), CapturingFrame->CaptureBaseData.SourceFrameNumberRenderThread));
+			RDG_GPU_STAT_SCOPE(GraphBuilder, MediaCapture_SyncPointPass);
+
+			// Initialize sync handlers only the first time. 
+			if (MediaCapture->bSyncHandlersInitialized == false)
+			{
+				MediaCapture->InitializeSyncHandlers_RenderThread();
+			}
+
+			// Add buffer output as a parameter to depend on the compute shader pass
+			typename CaptureType::PassParametersType* PassParameters = GraphBuilder.AllocParameters<typename CaptureType::PassParametersType>();
+			PassParameters->Resource = static_cast<typename CaptureType::FOutputResourceType>(OutputResource);
+
+			TWeakObjectPtr<UMediaCapture> WeakCapture = MediaCapture;
+
+			GraphBuilder.AddPass(
+				RDG_EVENT_NAME("MediaCaptureCopySyncPass"),
+				PassParameters,
+				ERDGPassFlags::Copy | ERDGPassFlags::NeverCull,
+				[WeakCapture, CapturingFrame](FRHICommandListImmediate& RHICmdList)
+				{
+					if (WeakCapture.IsValid() && WeakCapture->bIsActive)
+					{
+						// Get available sync handler to create a sync point
+						TSharedPtr<UMediaCapture::FMediaCaptureSyncData> SyncDataPtr = WeakCapture->GetAvailableSyncHandler();
+						if (ensure(SyncDataPtr))
+						{
+							// This will happen after the conversion pass has completed
+							RHICmdList.WriteGPUFence(SyncDataPtr->RHIFence);
+							SyncDataPtr->bIsBusy = true;
+
+							// Spawn a task that will wait (poll) and continue the process of providing a new texture
+							UE::Tasks::FTask Task = UE::Tasks::Launch(UE_SOURCE_LOCATION, [WeakCapture, CapturingFrame, SyncDataPtr]()
+								{
+									TRACE_CPUPROFILER_EVENT_SCOPE(UMediaCapture::SyncPass);
+
+									double WaitTime = 0.0;
+
+									{
+										FScopedDurationTimer Timer(WaitTime);
+									
+										// Wait until fence has been written (shader has completed)
+										while (SyncDataPtr->RHIFence->Poll() == false && WeakCapture->bIsActive)
+										{
+											FPlatformProcess::SleepNoStats(0);
+										}
+
+										SyncDataPtr->RHIFence->Clear();
+										SyncDataPtr->bIsBusy = false;
+									}
+
+									UE_LOG(LogMediaIOCore, VeryVerbose, TEXT("AsyncTask wait time: %f s"), WaitTime);
+
+									if (WeakCapture->bIsActive)
+									{
+										if (CVarMediaIOScheduleOnAnyThread.GetValueOnAnyThread() == 1)
+										{
+											OnReadbackComplete(FRHICommandListExecutor::GetImmediateCommandList(), WeakCapture.Get(), CapturingFrame);
+										}
+										else
+										{
+											++WeakCapture->WaitingForRenderCommandExecutionCounter;
+
+											ENQUEUE_RENDER_COMMAND(MediaOutputCaptureFrameCreateResources)([WeakCapture, CapturingFrame](FRHICommandList& RHICommandList)
+											{
+												if (WeakCapture.IsValid(false, true))
+												{
+													OnReadbackComplete(RHICommandList, WeakCapture.Get(), CapturingFrame);
+												}
+											});
+										}
+									}
+
+								}, LowLevelTasks::ETaskPriority::BackgroundHigh);
+
+							{
+								// Add the task to the capture's list of pending tasks.
+								FScopeLock ScopedLock(&WeakCapture->PendingReadbackTasksCriticalSection);
+								WeakCapture->PendingReadbackTasks.Add(MoveTemp(Task));
+							}
+						}
+					}
+				});
 		}
 
-		template<typename TFrameType>
-		static bool CaptureFrame(const FCaptureFrameArgs& Args, TFrameType* CapturingFrame)
+		static bool CaptureFrame(const FCaptureFrameArgs& Args, const TSharedPtr<UE::MediaCaptureData::FCaptureFrame> CapturingFrame)
 		{
 			RDG_GPU_STAT_SCOPE(Args.GraphBuilder, MediaCapture_CaptureFrame)
 
@@ -647,23 +784,24 @@ namespace UE::MediaCaptureData
 				return false;
 			}
 
-			// If true, we will need to go through our different shader to convert from source format to out format (i.e RGB to YUV)
-			const bool bRequiresFormatConversion = Args.MediaCapture->DesiredPixelFormat != Args.ResourceToCapture->GetFormat();
-
-			if (!Args.ResourceToCapture.IsValid())
+			if (!Args.HasValidResource())
 			{
 				UE_LOG(LogMediaIOCore, Error, TEXT("Can't grab the Texture to capture for '%s'."), *Args.MediaCapture->MediaOutputName);
 				return false;
 			}
 
+			// If true, we will need to go through our different shader to convert from source format to out format (i.e RGB to YUV)
+			const bool bRequiresFormatConversion = Args.MediaCapture->DesiredPixelFormat != Args.GetFormat();
+
 			// Validate pixel formats and sizes before pursuing
-			if (AreInputsValid(Args) == false)
+			if (!AreInputsValid(Args))
 			{
 				return false;
 			}
 
 			{
 				TRACE_CPUPROFILER_EVENT_SCOPE(UMediaCapture::LockDMATexture_RenderThread);
+				TRACE_CPUPROFILER_EVENT_SCOPE_TEXT(*FString::Printf(TEXT("LockDmaTexture Output Frame %d"), CapturingFrame->CaptureBaseData.SourceFrameNumberRenderThread));
 
 				if (CapturingFrame->IsTextureResource())
 				{
@@ -678,13 +816,19 @@ namespace UE::MediaCaptureData
 				FRHICopyTextureInfo CopyInfo;
 				FVector2D SizeU;
 				FVector2D SizeV;
-				GetCopyInfo(Args, Args.ResourceToCapture, CopyInfo, SizeU, SizeV);
+				GetCopyInfo(Args, CopyInfo, SizeU, SizeV);
 
-				// Register the source texture that we want to capture
-				const FRDGTextureRef SourceRGBTexture = Args.GraphBuilder.RegisterExternalTexture(CreateRenderTarget(Args.ResourceToCapture, TEXT("SourceTexture")));
+				FRDGTextureRef SourceRGBTexture = Args.RDGResourceToCapture;
+				
+				if (!SourceRGBTexture)
+				{
+					// If we weren't passed a rdg texture, register the external rhi texture.
+					SourceRGBTexture = Args.GraphBuilder.RegisterExternalTexture(CreateRenderTarget(Args.ResourceToCapture, TEXT("SourceTexture")));
+				}
 
 				// Register output resource used by the current capture method (texture or buffer)
-				typename TFrameType::FOutputResourceType OutputResource = CapturingFrame->RegisterResource(Args.GraphBuilder);
+				FRDGViewableResource* OutputResource = CapturingFrame->RegisterResource(Args.GraphBuilder);
+				
 				{
 					TRACE_CPUPROFILER_EVENT_SCOPE(UMediaCapture::GraphSetup);
 					SCOPED_DRAW_EVENTF(Args.GraphBuilder.RHICmdList, MediaCapture, TEXT("MediaCapture"));
@@ -694,9 +838,17 @@ namespace UE::MediaCaptureData
 					{
 						RDG_GPU_STAT_SCOPE(Args.GraphBuilder, MediaCapture_CustomCapture)
 						TRACE_CPUPROFILER_EVENT_SCOPE(UMediaCapture::CustomCapture);
-
-						Args.MediaCapture->OnCustomCapture_RenderingThread(Args.GraphBuilder, CapturingFrame->CaptureBaseData, CapturingFrame->UserData
-							, SourceRGBTexture, OutputResource, CopyInfo, SizeU, SizeV);
+						
+						if (CapturingFrame->IsTextureResource())
+						{
+							Args.MediaCapture->OnCustomCapture_RenderingThread(Args.GraphBuilder, CapturingFrame->CaptureBaseData, CapturingFrame->UserData
+								, SourceRGBTexture, static_cast<FRDGTextureRef>(OutputResource), CopyInfo, SizeU, SizeV);
+						}
+						else
+						{
+							Args.MediaCapture->OnCustomCapture_RenderingThread(Args.GraphBuilder, CapturingFrame->CaptureBaseData, CapturingFrame->UserData
+								, SourceRGBTexture, static_cast<FRDGBufferRef>(OutputResource), CopyInfo, SizeU, SizeV);
+						}
 					}
 					else
 					{
@@ -705,17 +857,15 @@ namespace UE::MediaCaptureData
 						AddConversionPass(Args, { SourceRGBTexture, bRequiresFormatConversion,  CopyInfo, SizeU, SizeV }, OutputResource);
 					}
 				}
+
 				
-				// If Capture implementation is not grabbing GPU resource directly, push a readback pass to access it from CPU
 				if (Args.MediaCapture->bShouldCaptureRHIResource == false)
 				{
 					RDG_GPU_STAT_SCOPE(Args.GraphBuilder, MediaCapture_Readback)
 					TRACE_CPUPROFILER_EVENT_SCOPE(UMediaCapture::EnqueueReadback);
 
 					CapturingFrame->EnqueueCopy(Args.GraphBuilder, OutputResource);
-
 					CapturingFrame->bReadbackRequested = true;
-					UE_LOG(LogMediaIOCore, VeryVerbose, TEXT("Requested copy for capture frame %d"), CapturingFrame->FrameId);
 				}
 				else
 				{
@@ -723,8 +873,19 @@ namespace UE::MediaCaptureData
 				}
 
 				Args.MediaCapture->FrameManager->MarkPending(*CapturingFrame);
-
 				++Args.MediaCapture->PendingFrameCount;
+
+				if (CVarMediaIOEnableExperimentalScheduling.GetValueOnAnyThread() == 1)
+				{
+					if (CapturingFrame->IsTextureResource())
+					{
+						AddSyncPointPass<FTextureCaptureFrame>(Args.GraphBuilder, Args.MediaCapture, CapturingFrame, OutputResource);
+					}
+					else
+					{
+						AddSyncPointPass<FBufferCaptureFrame>(Args.GraphBuilder, Args.MediaCapture, CapturingFrame, OutputResource);
+					}
+				}
 			}
 
 			return true;
@@ -732,15 +893,6 @@ namespace UE::MediaCaptureData
 	};
 }
 
-
-/* UMediaCapture::FCaptureBaseData
-*****************************************************************************/
-UMediaCapture::FCaptureBaseData::FCaptureBaseData()
-	: SourceFrameNumberRenderThread(0)
-{
-
-}
-
 /* FMediaCaptureOptions
 *****************************************************************************/
 FMediaCaptureOptions::FMediaCaptureOptions()
@@ -760,9 +912,8 @@ FMediaCaptureOptions::FMediaCaptureOptions()
 /* UMediaCapture
 *****************************************************************************/
 
-UMediaCapture::UMediaCapture(const FObjectInitializer& ObjectInitializer)
-	: Super(ObjectInitializer)
-	, ValidSourceGPUMask(FRHIGPUMask::All())
+UMediaCapture::UMediaCapture()
+	: ValidSourceGPUMask(FRHIGPUMask::All())
 	, bOutputResourcesInitialized(false)
 	, bShouldCaptureRHIResource(false)
 	, WaitingForRenderCommandExecutionCounter(0)
@@ -772,6 +923,49 @@ UMediaCapture::UMediaCapture(const FObjectInitializer& ObjectInitializer)
 
 UMediaCapture::~UMediaCapture() = default;
 
+const TSet<EPixelFormat>& UMediaCapture::GetSupportedRgbaSwizzleFormats()
+{
+	static TSet<EPixelFormat> SupportedFormats =
+		{
+			PF_A32B32G32R32F,
+			PF_B8G8R8A8,
+			PF_G8,
+			PF_G16,
+			PF_FloatRGB,
+			PF_FloatRGBA,
+			PF_R32_FLOAT,
+			PF_G16R16,
+			PF_G16R16F,
+			PF_G32R32F,
+			PF_A2B10G10R10,
+			PF_A16B16G16R16,
+			PF_R16F,
+			PF_FloatR11G11B10,
+			PF_A8,
+			PF_R32_UINT,
+			PF_R32_SINT,
+			PF_R16_UINT,
+			PF_R16_SINT,
+			PF_R16G16B16A16_UINT,
+			PF_R16G16B16A16_SINT,
+			PF_R5G6B5_UNORM,
+			PF_R8G8B8A8,
+			PF_A8R8G8B8,
+			PF_R8G8,
+			PF_R32G32B32A32_UINT,
+			PF_R16G16_UINT,
+			PF_R8_UINT,
+			PF_R8G8B8A8_UINT,
+			PF_R8G8B8A8_SNORM,
+			PF_R16G16B16A16_UNORM,
+			PF_R16G16B16A16_SNORM,
+			PF_R32G32_UINT,
+			PF_R8,
+		};
+	
+	return SupportedFormats;
+}
+
 void UMediaCapture::BeginDestroy()
 {
 	if (GetState() == EMediaCaptureState::Capturing || GetState() == EMediaCaptureState::Preparing)
@@ -808,13 +1002,37 @@ bool UMediaCapture::CaptureActiveSceneViewport(FMediaCaptureOptions CaptureOptio
 	return CaptureSceneViewport(FoundSceneViewport, CaptureOptions);
 }
 
-bool UMediaCapture::CaptureSceneViewport(TSharedPtr<FSceneViewport>& InSceneViewport, FMediaCaptureOptions InCaptureOptions)
+bool UMediaCapture::CaptureSceneViewport(TSharedPtr<FSceneViewport>& SceneViewport, FMediaCaptureOptions CaptureOptions)
 {
-	StopCapture(false);
+	using namespace UE::MediaCapture::Private;
+	return StartSourceCapture(MakeShared<FSceneViewportCaptureSource>(this, CaptureOptions, SceneViewport));
+}
 
+bool UMediaCapture::CaptureRHITexture(const FRHICaptureResourceDescription& ResourceDescription, FMediaCaptureOptions CaptureOptions)
+{
+	using namespace UE::MediaCapture::Private;
+	return StartSourceCapture(MakeShared<FRHIResourceCaptureSource>(this, CaptureOptions, ResourceDescription));
+}
+
+bool UMediaCapture::CaptureTextureRenderTarget2D(UTextureRenderTarget2D* RenderTarget, FMediaCaptureOptions CaptureOptions)
+{
+	using namespace UE::MediaCapture::Private;
+	return StartSourceCapture(MakeShared<FRenderTargetCaptureSource>(this, CaptureOptions, RenderTarget));
+}
+
+bool UMediaCapture::StartSourceCapture(TSharedPtr<UE::MediaCapture::Private::FCaptureSource> InSource)
+{
+	StopCapture(false);
+	
+	if (!InSource)
+	{
+		ensure(false);
+		return false;
+	}
+	
 	check(IsInGameThread());
 
-	DesiredCaptureOptions = InCaptureOptions;
+	DesiredCaptureOptions = InSource->CaptureOptions;
 
 	if (!ValidateMediaOutput())
 	{
@@ -822,23 +1040,22 @@ bool UMediaCapture::CaptureSceneViewport(TSharedPtr<FSceneViewport>& InSceneView
 		return false;
 	}
 
-	CacheMediaOutput(EMediaCaptureSourceType::SCENE_VIEWPORT);
+	CacheMediaOutput(InSource->GetSourceType());
 
 	if (bUseRequestedTargetSize)
 	{
-		DesiredSize = InSceneViewport->GetSize();
+		DesiredSize = InSource->GetSize();
 	}
 	else if (DesiredCaptureOptions.bResizeSourceBuffer)
 	{
-		SetFixedViewportSize(InSceneViewport);
+		InSource->ResizeSourceBuffer(DesiredSize);
 	}
 
 	CacheOutputOptions();
 
-	const bool bCurrentlyCapturing = false;
-	if (!MediaCaptureDetails::ValidateSceneViewport(InSceneViewport, DesiredCaptureOptions, DesiredSize, DesiredPixelFormat, bCurrentlyCapturing))
+	constexpr bool bCurrentlyCapturing = false;
+	if (!InSource->ValidateSource(DesiredSize, DesiredPixelFormat, bCurrentlyCapturing))
 	{
-		ResetFixedViewportSize(InSceneViewport, false);
 		MediaCaptureDetails::ShowSlateNotification();
 		return false;
 	}
@@ -848,11 +1065,11 @@ bool UMediaCapture::CaptureSceneViewport(TSharedPtr<FSceneViewport>& InSceneView
 	bool bInitialized = InitializeCapture();
 	if (bInitialized)
 	{
-		bInitialized = PostInitializeCaptureViewport(InSceneViewport);
+		bInitialized = InSource->PostInitialize();
 	}
 
 	// This could have been updated by the initialization done by the implementation
-	bShouldCaptureRHIResource = ShouldCaptureRHIResource(); 
+	bShouldCaptureRHIResource = ShouldCaptureRHIResource();
 
 	if (bInitialized)
 	{
@@ -862,13 +1079,31 @@ bool UMediaCapture::CaptureSceneViewport(TSharedPtr<FSceneViewport>& InSceneView
 
 	if (bInitialized)
 	{
-		//no lock required, the command on the render thread is not active
-		CapturingSceneViewport = InSceneViewport;
-		FCoreDelegates::OnEndFrame.AddUObject(this, &UMediaCapture::OnEndFrame_GameThread);
+		//no lock required the command on the render thread is not active yet
+		CaptureSource = MoveTemp(InSource);
+		
+		if (DesiredCaptureOptions.CapturePhase != EMediaCapturePhase::EndFrame)
+		{
+			ViewExtension = FSceneViewExtensions::NewExtension<FMediaCaptureSceneViewExtension>(this, DesiredCaptureOptions.CapturePhase);
+		}
+
+		if (ViewExtension)
+		{
+			ensure(CaptureSource->GetCaptureType() == UE::MediaCapture::Private::ECaptureType::Immediate);
+		}
+
+		if (CaptureSource->GetCaptureType() == UE::MediaCapture::Private::ECaptureType::Immediate)
+		{
+			// Immediate capture requires us to prepare frame info in OnBeginFrame
+			FCoreDelegates::OnBeginFrame.AddUObject(this, &UMediaCapture::OnBeginFrame_GameThread);
+		}
+		else
+		{
+			FCoreDelegates::OnEndFrame.AddUObject(this, &UMediaCapture::OnEndFrame_GameThread);
+		}
 	}
 	else
 	{
-		ResetFixedViewportSize(InSceneViewport, false);
 		SetState(EMediaCaptureState::Stopped);
 		MediaCaptureDetails::ShowSlateNotification();
 	}
@@ -876,177 +1111,48 @@ bool UMediaCapture::CaptureSceneViewport(TSharedPtr<FSceneViewport>& InSceneView
 #if WITH_EDITOR
 	MediaCaptureAnalytics::SendCaptureEvent(GetCaptureSourceType());
 #endif
-	
+
+	bIsActive = true;
+
 	return bInitialized;
 }
 
-bool UMediaCapture::CaptureRHITexture(const FRHICaptureResourceDescription& ResourceDescription, FMediaCaptureOptions CaptureOptions)
+void UMediaCapture::CacheMediaOutput(EMediaCaptureSourceType InSourceType)
 {
-	StopCapture(false);
-
-	check(IsInGameThread());
+	check(MediaOutput);
+	CaptureSourceType = InSourceType;
+	DesiredSize = MediaOutput->GetRequestedSize();
+	bUseRequestedTargetSize = DesiredSize == UMediaOutput::RequestCaptureSourceSize;
+	DesiredPixelFormat = MediaOutput->GetRequestedPixelFormat();
+	ConversionOperation = MediaOutput->GetConversionOperation(InSourceType);
+}
 
-	DesiredCaptureOptions = CaptureOptions;
+void UMediaCapture::CacheOutputOptions()
+{
+	DesiredOutputSize = GetOutputSize(DesiredSize, ConversionOperation);
+	DesiredOutputResourceType = GetOutputResourceType(ConversionOperation);
+	DesiredOutputPixelFormat = GetOutputPixelFormat(DesiredPixelFormat, ConversionOperation);
+	DesiredOutputBufferDescription = GetOutputBufferDescription(ConversionOperation);
+	MediaOutputName = *MediaOutput->GetName();
+	bShouldCaptureRHIResource = ShouldCaptureRHIResource();
+}
 
-	if (!ValidateMediaOutput())
+FIntPoint UMediaCapture::GetOutputSize(const FIntPoint & InSize, const EMediaCaptureConversionOperation & InConversionOperation) const
+{
+	switch (InConversionOperation)
 	{
-		MediaCaptureDetails::ShowSlateNotification();
-		return false;
+	case EMediaCaptureConversionOperation::RGBA8_TO_YUV_8BIT:
+		return FIntPoint(InSize.X / 2, InSize.Y);
+	case EMediaCaptureConversionOperation::RGB10_TO_YUVv210_10BIT:
+		// Padding aligned on 48 (16 and 6 at the same time)
+		return FIntPoint((((InSize.X + 47) / 48) * 48) / 6, InSize.Y);
+	case EMediaCaptureConversionOperation::CUSTOM:
+		return GetCustomOutputSize(InSize);
+	case EMediaCaptureConversionOperation::NONE:
+	default:
+		return InSize;
 	}
-
-	CacheMediaOutput(EMediaCaptureSourceType::RHI_RESOURCE);
-
-	// If output is configured to capture resource size directly (not enforced), configure desired size to be the same as resource
-	if (bUseRequestedTargetSize)
-	{
-		DesiredSize = ResourceDescription.ResourceSize;
-	}
-
-	CacheOutputOptions();
-
-	// Can't validate resource to capture as it will be given everytime a capture is requested
-	SetState(EMediaCaptureState::Preparing);
-
-	bool bInitialized = InitializeCapture();
-	if (bInitialized)
-	{
-		bInitialized = PostInitializeCaptureRHIResource(ResourceDescription);
-	}
-
-	// This could have been updated by the call to implementation initialization
-	bShouldCaptureRHIResource = ShouldCaptureRHIResource();
-
-	if (bInitialized)
-	{
-		InitializeOutputResources(MediaOutput->NumberOfTextureBuffers);
-		bInitialized = GetState() != EMediaCaptureState::Stopped;
-	}
-
-	if (bInitialized)
-	{
-		// Used to queue capture data production from the game thread to be consumed when frame capture is requested on render thread
-		FCoreDelegates::OnBeginFrame.AddUObject(this, &UMediaCapture::OnBeginFrame_GameThread);
-	}
-	else
-	{
-		SetState(EMediaCaptureState::Stopped);
-		MediaCaptureDetails::ShowSlateNotification();
-	}
-
-#if WITH_EDITOR
-	MediaCaptureAnalytics::SendCaptureEvent(GetCaptureSourceType());
-#endif
-
-	return bInitialized;
-}
-
-bool UMediaCapture::CaptureTextureRenderTarget2D(UTextureRenderTarget2D* InRenderTarget2D, FMediaCaptureOptions CaptureOptions)
-{
-	StopCapture(false);
-
-	check(IsInGameThread());
-
-	DesiredCaptureOptions = CaptureOptions; 
-	
-	if (!ValidateMediaOutput())
-	{
-		MediaCaptureDetails::ShowSlateNotification();
-		return false;
-	}
-
-	CacheMediaOutput(EMediaCaptureSourceType::RENDER_TARGET);
-
-	if (bUseRequestedTargetSize)
-	{
-		DesiredSize = FIntPoint(InRenderTarget2D->SizeX, InRenderTarget2D->SizeY);
-	}
-	else if (DesiredCaptureOptions.bResizeSourceBuffer)
-	{
-		InRenderTarget2D->ResizeTarget(DesiredSize.X, DesiredSize.Y);
-	}
-
-	CacheOutputOptions();
-
-	const bool bCurrentlyCapturing = false;
-	if (!MediaCaptureDetails::ValidateTextureRenderTarget2D(InRenderTarget2D, DesiredCaptureOptions, DesiredSize, DesiredPixelFormat, bCurrentlyCapturing))
-	{
-		MediaCaptureDetails::ShowSlateNotification();
-		return false;
-	}
-
-	SetState(EMediaCaptureState::Preparing);
-
-	bool bInitialized = InitializeCapture();
-	if (bInitialized)
-	{
-		bInitialized = PostInitializeCaptureRenderTarget(InRenderTarget2D);
-	}
-
-	// This could have been updated by the initialization done by the implementation
-	bShouldCaptureRHIResource = ShouldCaptureRHIResource();
-
-	if (bInitialized)
-	{
-		InitializeOutputResources(MediaOutput->NumberOfTextureBuffers);
-		bInitialized = GetState() != EMediaCaptureState::Stopped;
-	}
-
-	if (bInitialized)
-	{
-		//no lock required the command on the render thread is not active yet
-		CapturingRenderTarget = InRenderTarget2D;
-
-		FCoreDelegates::OnEndFrame.AddUObject(this, &UMediaCapture::OnEndFrame_GameThread);
-	}
-	else
-	{
-		SetState(EMediaCaptureState::Stopped);
-		MediaCaptureDetails::ShowSlateNotification();
-	}
-
-#if WITH_EDITOR
-	MediaCaptureAnalytics::SendCaptureEvent(GetCaptureSourceType());
-#endif
-	
-	return bInitialized;
-}
-
-void UMediaCapture::CacheMediaOutput(EMediaCaptureSourceType InSourceType)
-{
-	check(MediaOutput);
-	CaptureSourceType = InSourceType;
-	DesiredSize = MediaOutput->GetRequestedSize();
-	bUseRequestedTargetSize = DesiredSize == UMediaOutput::RequestCaptureSourceSize;
-	DesiredPixelFormat = MediaOutput->GetRequestedPixelFormat();
-	ConversionOperation = MediaOutput->GetConversionOperation(InSourceType);
-}
-
-void UMediaCapture::CacheOutputOptions()
-{
-	DesiredOutputSize = GetOutputSize(DesiredSize, ConversionOperation);
-	DesiredOutputResourceType = GetOutputResourceType(ConversionOperation);
-	DesiredOutputPixelFormat = GetOutputPixelFormat(DesiredPixelFormat, ConversionOperation);
-	DesiredOutputBufferDescription = GetOutputBufferDescription(ConversionOperation);
-	MediaOutputName = *MediaOutput->GetName();
-	bShouldCaptureRHIResource = ShouldCaptureRHIResource();
-}
-
-FIntPoint UMediaCapture::GetOutputSize(const FIntPoint & InSize, const EMediaCaptureConversionOperation & InConversionOperation) const
-{
-	switch (InConversionOperation)
-	{
-	case EMediaCaptureConversionOperation::RGBA8_TO_YUV_8BIT:
-		return FIntPoint(InSize.X / 2, InSize.Y);
-	case EMediaCaptureConversionOperation::RGB10_TO_YUVv210_10BIT:
-		// Padding aligned on 48 (16 and 6 at the same time)
-		return FIntPoint((((InSize.X + 47) / 48) * 48) / 6, InSize.Y);
-	case EMediaCaptureConversionOperation::CUSTOM:
-		return GetCustomOutputSize(InSize);
-	case EMediaCaptureConversionOperation::NONE:
-	default:
-		return InSize;
-	}
-}
+}
 
 EPixelFormat UMediaCapture::GetOutputPixelFormat(const EPixelFormat & InPixelFormat, const EMediaCaptureConversionOperation & InConversionOperation) const
 {
@@ -1086,9 +1192,9 @@ FRDGBufferDesc UMediaCapture::GetOutputBufferDescription(EMediaCaptureConversion
 	}
 }
 
-bool UMediaCapture::UpdateSceneViewport(TSharedPtr<FSceneViewport>& InSceneViewport)
+bool UMediaCapture::UpdateSource(TSharedPtr<UE::MediaCapture::Private::FCaptureSource> InCaptureSource)
 {
-	if (!MediaCaptureDetails::ValidateIsCapturing(*this))
+	if (!MediaCaptureDetails::ValidateIsCapturing(this))
 	{
 		StopCapture(false);
 		return false;
@@ -1096,23 +1202,28 @@ bool UMediaCapture::UpdateSceneViewport(TSharedPtr<FSceneViewport>& InSceneViewp
 
 	check(IsInGameThread());
 
+	const TSharedPtr<UE::MediaCapture::Private::FCaptureSource> PreviousCaptureSource = CaptureSource;
+	CaptureSource = MoveTemp(InCaptureSource);
+
 	if (!bUseRequestedTargetSize && DesiredCaptureOptions.bResizeSourceBuffer)
 	{
-		SetFixedViewportSize(InSceneViewport);
+		CaptureSource->ResizeSourceBuffer(DesiredSize);
 	}
 
 	const bool bCurrentlyCapturing = true;
-	if (!MediaCaptureDetails::ValidateSceneViewport(InSceneViewport, DesiredCaptureOptions, DesiredSize, DesiredPixelFormat, bCurrentlyCapturing))
+	if (!CaptureSource->ValidateSource(DesiredSize, DesiredPixelFormat, bCurrentlyCapturing))
 	{
-		ResetFixedViewportSize(InSceneViewport, false);
+		constexpr bool bFlushRenderingCommands = false;
+		CaptureSource->ResetSourceBufferSize(bFlushRenderingCommands);
 		StopCapture(false);
 		MediaCaptureDetails::ShowSlateNotification();
 		return false;
 	}
 
-	if (!UpdateSceneViewportImpl(InSceneViewport))
+	if (!CaptureSource->UpdateSourceImpl())
 	{
-		ResetFixedViewportSize(InSceneViewport, false);
+		constexpr bool bFlushRenderingCommands = false;
+		CaptureSource->ResetSourceBufferSize(bFlushRenderingCommands);
 		StopCapture(false);
 		MediaCaptureDetails::ShowSlateNotification();
 		return false;
@@ -1120,61 +1231,26 @@ bool UMediaCapture::UpdateSceneViewport(TSharedPtr<FSceneViewport>& InSceneViewp
 
 	{
 		FScopeLock Lock(&AccessingCapturingSource);
-		while(WaitingForRenderCommandExecutionCounter.load() > 0)
-		{
-			FlushRenderingCommands();
-		}
-		ResetFixedViewportSize(CapturingSceneViewport.Pin(), true);
-		CapturingSceneViewport = InSceneViewport;
-		CapturingRenderTarget = nullptr;
+
+		WaitForPendingTasks();
+		
+		constexpr bool bFlushRenderingCommands = false;
+		PreviousCaptureSource->ResetSourceBufferSize(bFlushRenderingCommands);
 	}
 
 	return true;
 }
 
-bool UMediaCapture::UpdateTextureRenderTarget2D(UTextureRenderTarget2D * InRenderTarget2D)
+bool UMediaCapture::UpdateSceneViewport(TSharedPtr<FSceneViewport>& InSceneViewport)
 {
-	if (!MediaCaptureDetails::ValidateIsCapturing(*this))
-	{
-		StopCapture(false);
-		MediaCaptureDetails::ShowSlateNotification();
-		return false;
-	}
-
-	check(IsInGameThread());
-
-	if (!bUseRequestedTargetSize && DesiredCaptureOptions.bResizeSourceBuffer)
-	{
-		InRenderTarget2D->ResizeTarget(DesiredSize.X, DesiredSize.Y);
-	}
-
-	const bool bCurrentlyCapturing = true;
-	if (!MediaCaptureDetails::ValidateTextureRenderTarget2D(InRenderTarget2D, DesiredCaptureOptions, DesiredSize, DesiredPixelFormat, bCurrentlyCapturing))
-	{
-		StopCapture(false);
-		MediaCaptureDetails::ShowSlateNotification();
-		return false;
-	}
-
-	if (!UpdateRenderTargetImpl(InRenderTarget2D))
-	{
-		StopCapture(false);
-		MediaCaptureDetails::ShowSlateNotification();
-		return false;
-	}
-
-	{
-		FScopeLock Lock(&AccessingCapturingSource);
-		while (WaitingForRenderCommandExecutionCounter.load() > 0)
-		{
-			FlushRenderingCommands();
-		}
-		ResetFixedViewportSize(CapturingSceneViewport.Pin(), true);
-		CapturingRenderTarget = InRenderTarget2D;
-		CapturingSceneViewport.Reset();
-	}
+	check(CaptureSource);
+	return UpdateSource(MakeShared<UE::MediaCapture::Private::FSceneViewportCaptureSource>(this, CaptureSource->CaptureOptions, InSceneViewport));
+}
 
-	return true;
+bool UMediaCapture::UpdateTextureRenderTarget2D(UTextureRenderTarget2D* InRenderTarget2D)
+{
+	check(CaptureSource);
+	return UpdateSource(MakeShared<UE::MediaCapture::Private::FRenderTargetCaptureSource>(this, CaptureSource->CaptureOptions, InRenderTarget2D));
 }
 
 void UMediaCapture::StopCapture(bool bAllowPendingFrameToBeProcess)
@@ -1195,15 +1271,14 @@ void UMediaCapture::StopCapture(bool bAllowPendingFrameToBeProcess)
 			//Do not flush when auto stopping to avoid hitches.
 			if(DesiredCaptureOptions.bAutostopOnCapture != true)
 			{
-				while (WaitingForRenderCommandExecutionCounter.load() > 0)
-				{
-					FlushRenderingCommands();
-				}
+				WaitForPendingTasks();
 			}
 		}
 	}
 	else
 	{
+		bIsActive = false;
+
 		if (GetState() != EMediaCaptureState::Stopped)
 		{
 			SetState(EMediaCaptureState::Stopped);
@@ -1215,17 +1290,24 @@ void UMediaCapture::StopCapture(bool bAllowPendingFrameToBeProcess)
 			{
 				FlushRenderingCommands();
 			}
+			
 			StopCaptureImpl(bAllowPendingFrameToBeProcess);
-			ResetFixedViewportSize(CapturingSceneViewport.Pin(), false);
 
-			CapturingRenderTarget = nullptr;
-			CapturingSceneViewport.Reset();
+			if (CaptureSource)
+			{
+				constexpr bool bFlushRenderingCommands = false;
+				CaptureSource->ResetSourceBufferSize(bFlushRenderingCommands);
+				CaptureSource.Reset();
+			}
+			
 			DesiredSize = FIntPoint(1280, 720);
 			DesiredPixelFormat = EPixelFormat::PF_A2B10G10R10;
 			DesiredOutputSize = FIntPoint(1280, 720);
 			DesiredOutputPixelFormat = EPixelFormat::PF_A2B10G10R10;
 			DesiredCaptureOptions = FMediaCaptureOptions();
 			ConversionOperation = EMediaCaptureConversionOperation::NONE;
+			ViewExtension.Reset();
+			
 			MediaOutputName.Reset();
 		}
 	}
@@ -1240,6 +1322,25 @@ void UMediaCapture::SetMediaOutput(UMediaOutput* InMediaOutput)
 }
 
 void UMediaCapture::CaptureImmediate_RenderThread(FRDGBuilder& GraphBuilder, FRHITexture* InSourceTexture)
+{
+	UE::MediaCaptureData::FCaptureFrameArgs CaptureArgs{GraphBuilder};
+	CaptureArgs.MediaCapture = this;
+    CaptureArgs.DesiredSize = DesiredSize;
+    CaptureArgs.ResourceToCapture = InSourceTexture;
+	CaptureImmediate_RenderThread(CaptureArgs);
+}
+
+
+void UMediaCapture::CaptureImmediate_RenderThread(FRDGBuilder& GraphBuilder,  FRDGTextureRef InSourceTextureRef)
+{
+	UE::MediaCaptureData::FCaptureFrameArgs CaptureArgs{GraphBuilder};
+	CaptureArgs.MediaCapture = this;
+	CaptureArgs.DesiredSize = DesiredSize;
+	CaptureArgs.RDGResourceToCapture = InSourceTextureRef;
+	CaptureImmediate_RenderThread(CaptureArgs);
+}
+
+void UMediaCapture::CaptureImmediate_RenderThread(const UE::MediaCaptureData::FCaptureFrameArgs& Args)
 {
 	using namespace UE::MediaCaptureData;
 	TRACE_CPUPROFILER_EVENT_SCOPE(UMediaCapture::CaptureImmediate_RenderThread);
@@ -1263,7 +1364,7 @@ void UMediaCapture::CaptureImmediate_RenderThread(FRDGBuilder& GraphBuilder, FRH
 		return;
 	}
 
-	if (CaptureSourceType != EMediaCaptureSourceType::RHI_RESOURCE)
+	if (CaptureSource && CaptureSource->GetCaptureType() != UE::MediaCapture::Private::ECaptureType::Immediate)
 	{
 		UE_LOG(LogMediaIOCore, Error, TEXT("[%s] - Trying to capture a RHI resource with another capture type."), *MediaOutputName);
 		SetState(EMediaCaptureState::Error);
@@ -1305,19 +1406,26 @@ void UMediaCapture::CaptureImmediate_RenderThread(FRDGBuilder& GraphBuilder, FRH
 	}
 		
 	// Get next available frame from the store. Can be invalid.
-	FCaptureFrame* CapturingFrame = nullptr;
+	TSharedPtr<FCaptureFrame> CapturingFrame;
 	if (GetState() != EMediaCaptureState::StopRequested)
 	{
-		CapturingFrame = FrameManager->GetNextAvailable();
+		CapturingFrame = FrameManager->GetNextAvailable<FCaptureFrame>();
 	}
 
-	if (CapturingFrame == nullptr && GetState() != EMediaCaptureState::StopRequested)
+	if (!CapturingFrame && GetState() != EMediaCaptureState::StopRequested)
 	{
 		if (DesiredCaptureOptions.OverrunAction == EMediaCaptureOverrunAction::Flush)
 		{
-			UE_LOG(LogMediaIOCore, Error, TEXT("[%s] - No frames available for capture. This should not happen."), *MediaOutputName);
-			SetState(EMediaCaptureState::Error);
-			return;
+			WaitForSingleExperimentalSchedulingTaskToComplete();
+
+			CapturingFrame = FrameManager->GetNextAvailable<FCaptureFrame>();
+
+			if (!CapturingFrame)
+			{
+				UE_LOG(LogMediaIOCore, Error, TEXT("[%s] - No frames available for capture. This should not happen."), *MediaOutputName);
+				SetState(EMediaCaptureState::Error);
+				return;
+			}
 		}
 		else
 		{
@@ -1326,17 +1434,9 @@ void UMediaCapture::CaptureImmediate_RenderThread(FRDGBuilder& GraphBuilder, FRH
 		}
 	}
 
-	// Peek next pending before capturing to avoid getting the newly captured one. 
-	// It shouldn't happen but if it happens it will start waiting for the capture to complete
-	FCaptureFrame* NextPending = FrameManager->PeekNextPending();
-	if (NextPending && CapturingFrame)
-	{
-		ensure(NextPending->FrameId != CapturingFrame->FrameId);
-	}
-
 	PrintFrameState();
 
-	UE_LOG(LogMediaIOCore, Verbose, TEXT("[%s] - Capturing frame %d"), *MediaOutputName, CapturingFrame ? CapturingFrame->FrameId : -1);
+	UE_LOG(LogMediaIOCore, Verbose, TEXT("[%s - %s] - Capturing frame %d"), *MediaOutputName, *FThreadManager::GetThreadName(FPlatformTLS::GetCurrentThreadId()), CapturingFrame ? CapturingFrame->FrameId : -1);
 	if (CapturingFrame)
 	{
 		// Prepare frame to capture
@@ -1346,41 +1446,43 @@ void UMediaCapture::CaptureImmediate_RenderThread(FRDGBuilder& GraphBuilder, FRH
 		CapturingFrame->CaptureBaseData.SourceFrameNumber = ++CaptureRequestCount;
 		CapturingFrame->UserData = MoveTemp(NextCaptureData.UserData);
 
-		ProcessCapture_RenderThread(GraphBuilder, this, CapturingFrame, InSourceTexture, DesiredSize);
+		ProcessCapture_RenderThread(CapturingFrame, Args);
+	}
+
+	// If CVarMediaIOEnableExperimentalScheduling is enabled, it means that a sync point pass was added, and the ready frame processing will be done later
+	if (CVarMediaIOEnableExperimentalScheduling.GetValueOnAnyThread() == 0)
+	{
+		if (TSharedPtr<FCaptureFrame> NextPending = FrameManager->PeekNextPending<FCaptureFrame>())
+		{
+			UE_LOG(LogMediaIOCore, Verbose, TEXT("[%s - %s] - Processing pending frame %d"), *MediaOutputName, *FThreadManager::GetThreadName(FPlatformTLS::GetCurrentThreadId()), NextPending ? NextPending->FrameId : -1);
+			ProcessReadyFrame_RenderThread(Args.GraphBuilder.RHICmdList, this, NextPending);
+		}
 	}
+}
 
-	UE_LOG(LogMediaIOCore, Verbose, TEXT("[%s] - Processing pending frame %d"), *MediaOutputName, NextPending ? NextPending->FrameId : -1);
-	if(NextPending)
+UTextureRenderTarget2D* UMediaCapture::GetTextureRenderTarget() const
+{
+	if (CaptureSource)
 	{
-		ProcessReadyFrame_RenderThread(GraphBuilder.RHICmdList, this, NextPending);
+		return CaptureSource->GetRenderTarget();
 	}
+	return nullptr;
 }
 
-const FString& UMediaCapture::GetCaptureSourceType() const
+TSharedPtr<FSceneViewport> UMediaCapture::GetCapturingSceneViewport() const
 {
-	switch (CaptureSourceType)
+	if (CaptureSource)
 	{
-		case EMediaCaptureSourceType::SCENE_VIEWPORT:
-		{
-			static FString Viewport = TEXT("SceneViewport");
-			return Viewport;
-		}
-		case EMediaCaptureSourceType::RENDER_TARGET:
-		{
-			static FString RenderTarget = TEXT("RenderTarget");
-			return RenderTarget;
-		}
-		case EMediaCaptureSourceType::RHI_RESOURCE:
-		{
-			static FString RHIResource = TEXT("RHI");
-			return RHIResource;
-		}
-		default:
-		{
-			static FString Invalid = TEXT("Invalid");
-			return Invalid;
-		}
+		return CaptureSource->GetSceneViewport();
 	}
+	return nullptr;
+}
+
+FString UMediaCapture::GetCaptureSourceType() const
+{
+	const UEnum* EnumType = StaticEnum<EMediaCaptureSourceType>();
+	check(EnumType);
+	return EnumType->GetNameStringByValue((int8)CaptureSourceType);
 }
 
 void UMediaCapture::SetState(EMediaCaptureState InNewState)
@@ -1405,30 +1507,239 @@ void UMediaCapture::SetState(EMediaCaptureState InNewState)
 			});
 		}
 	}
-}
+}
+
+void UMediaCapture::BroadcastStateChanged()
+{
+	OnStateChanged.Broadcast();
+	OnStateChangedNative.Broadcast();
+}
+
+void UMediaCapture::SetFixedViewportSize(TSharedPtr<FSceneViewport> InSceneViewport, FIntPoint InSize)
+{
+	InSceneViewport->SetFixedViewportSize(InSize.X, InSize.Y);
+	bViewportHasFixedViewportSize = true;
+}
+
+void UMediaCapture::ResetFixedViewportSize(TSharedPtr<FSceneViewport> InViewport, bool bInFlushRenderingCommands)
+{
+	if (bViewportHasFixedViewportSize && InViewport.IsValid())
+	{
+		if (bInFlushRenderingCommands)
+		{
+			WaitForPendingTasks();
+		}
+		InViewport->SetFixedViewportSize(0, 0);
+		bViewportHasFixedViewportSize = false;
+	}
+}
+
+void UMediaCapture::WaitForSingleExperimentalSchedulingTaskToComplete()
+{
+	if (CVarMediaIOEnableExperimentalScheduling.GetValueOnAnyThread())
+	{
+		// Presumably the rendering thread could be in the process of dispatching the task
+		FScopeLock ScopedLock(&PendingReadbackTasksCriticalSection);
+		if (PendingReadbackTasks.Num())
+		{
+			auto It = PendingReadbackTasks.CreateIterator();
+
+			// Go through the list to find a pending task.
+			while (It && It->IsCompleted())
+			{
+				It.RemoveCurrent();
+				++It;
+			}
+
+			if (It)
+			{
+				// We have a pending task, wait until it's completeted.
+				It->BusyWait();
+			}
+		}
+
+		if (!CVarMediaIOScheduleOnAnyThread.GetValueOnAnyThread())
+		{
+			// We flush after task completion in case a render thread task was launched.
+			if (IsInGameThread())
+			{
+				FlushRenderingCommands();
+			}
+		}
+	}
+}
+
+void UMediaCapture::CleanupCompletedExperimentalSchedulingTasks()
+{
+	FScopeLock ScopedLock(&PendingReadbackTasksCriticalSection);
+
+	for (auto It = PendingReadbackTasks.CreateIterator(); It; ++It)
+	{
+		if (It->IsCompleted())
+		{
+			It.RemoveCurrent();
+		}
+	}
+}
+
+void UMediaCapture::WaitForAllExperimentalSchedulingTasksToComplete()
+{
+	if (CVarMediaIOEnableExperimentalScheduling.GetValueOnAnyThread())
+	{
+		FScopeLock ScopedLock(&PendingReadbackTasksCriticalSection);
+
+		// Clean up completed tasks.
+		for (auto It = PendingReadbackTasks.CreateIterator(); It; ++It)
+		{
+			if (It->IsCompleted())
+			{
+				It.RemoveCurrent();
+			}
+			else
+			{
+				It->BusyWait();
+				It.RemoveCurrent();
+			}
+		}
+
+		
+		if (!CVarMediaIOScheduleOnAnyThread.GetValueOnGameThread())
+		{
+			// This code might have dispatched a task on the render thread, so we need to wait again
+			if (IsInGameThread())
+			{
+				FlushRenderingCommands();
+			}
+		}
+	}
+}
+
+void UMediaCapture::WaitForPendingTasks()
+{
+	while (WaitingForRenderCommandExecutionCounter.load() > 0)
+	{
+		FlushRenderingCommands();
+	}
+
+	WaitForAllExperimentalSchedulingTasksToComplete();
+}
+
+void UMediaCapture::ProcessCapture_GameThread()
+{
+	using namespace UE::MediaCaptureData;
+	
+	CleanupCompletedExperimentalSchedulingTasks();
+
+	// Acquire a frame
+	TSharedPtr<FCaptureFrame> CapturingFrame;
+	if(GetState() != EMediaCaptureState::StopRequested)
+	{
+		CapturingFrame = FrameManager->GetNextAvailable<FCaptureFrame>();
+	}
+
+	// Handle frame overrun (couldn't acquire a frame
+	if (!CapturingFrame && GetState() != EMediaCaptureState::StopRequested)
+	{
+		if (DesiredCaptureOptions.OverrunAction == EMediaCaptureOverrunAction::Flush)
+		{
+			TRACE_CPUPROFILER_EVENT_SCOPE(MediaCapture::FlushRenderingCommands);
+			UE_LOG(LogMediaIOCore, VeryVerbose, TEXT("[%s] - Flushing commands."), *MediaOutputName);
+
+			FlushRenderingCommands();
+			
+			WaitForSingleExperimentalSchedulingTaskToComplete();
+
+			// After flushing, we should have access to an available frame, if not, it's not expected.
+			CapturingFrame = FrameManager->GetNextAvailable<FCaptureFrame>();
+			if (!CapturingFrame)
+			{
+				UE_LOG(LogMediaIOCore, Error, TEXT("[%s] - Flushing commands didn't give us back available frames to process.") , *MediaOutputName);
+				StopCapture(false);
+				return;
+			}
+		}
+		else if (DesiredCaptureOptions.OverrunAction == EMediaCaptureOverrunAction::Skip)
+		{
+			// Selected options is to skip capturing a frame if overrun happens
+			UE_LOG(LogMediaIOCore, Verbose, TEXT("[%s] - No frames available for capture. Skipping"), *MediaOutputName);
+		}
+	}
 
-void UMediaCapture::BroadcastStateChanged()
-{
-	OnStateChanged.Broadcast();
-	OnStateChangedNative.Broadcast();
-}
+	// Initialize capture frame
+	InitializeCaptureFrame(CapturingFrame);
 
-void UMediaCapture::SetFixedViewportSize(TSharedPtr<FSceneViewport> InSceneViewport)
-{
-	InSceneViewport->SetFixedViewportSize(DesiredSize.X, DesiredSize.Y);
-	bViewportHasFixedViewportSize = true;
+	PrintFrameState();
+	
+	PrepareAndDispatchCapture_GameThread(CapturingFrame);
 }
 
-void UMediaCapture::ResetFixedViewportSize(TSharedPtr<FSceneViewport> InViewport, bool bInFlushRenderingCommands)
+void UMediaCapture::PrepareAndDispatchCapture_GameThread(const TSharedPtr<UE::MediaCaptureData::FCaptureFrame>& CapturingFrame)
 {
-	if (bViewportHasFixedViewportSize && InViewport.IsValid())
+	using namespace UE::MediaCaptureData;
+	
+	// Init variables for ENQUEUE_RENDER_COMMAND.
+	//The Lock only synchronize while we are copying the value to the enqueue. The viewport and the rendertarget may change while we are in the enqueue command.
 	{
-		if (bInFlushRenderingCommands && WaitingForRenderCommandExecutionCounter.load() > 0)
+		FScopeLock Lock(&AccessingCapturingSource);
+
+		FIntPoint InDesiredSize = DesiredSize;
+		UMediaCapture* InMediaCapture = this;
+
+		if (CaptureSource)
 		{
-			FlushRenderingCommands();
+			++WaitingForRenderCommandExecutionCounter;
+
+			const FRHIGPUMask SourceGPUMask = ValidSourceGPUMask;
+
+			UE_LOG(LogMediaIOCore, Verbose, TEXT("[%s] - Queuing frame to capture %d"), *InMediaCapture->MediaOutputName, CapturingFrame ? CapturingFrame->FrameId : -1);
+
+			// RenderCommand to be executed on the RenderThread
+			ENQUEUE_RENDER_COMMAND(FMediaOutputCaptureFrameCreateTexture)(
+				[InMediaCapture, CapturingFrame, InDesiredSize, SourceGPUMask](FRHICommandListImmediate& RHICmdList)
+			{
+				SCOPED_GPU_MASK(RHICmdList, SourceGPUMask);
+
+				TSharedPtr<FCaptureFrame> NextPending = InMediaCapture->FrameManager->PeekNextPending<FCaptureFrame>();
+				if (NextPending && CapturingFrame)
+				{
+					ensure(NextPending->FrameId != CapturingFrame->FrameId);
+				}
+				
+				// Capture frame
+				{
+					FRDGBuilder GraphBuilder(RHICmdList);
+					
+					FTexture2DRHIRef SourceTexture = InMediaCapture->CaptureSource->GetSourceTextureForInput_RenderThread(RHICmdList);
+					
+					FCaptureFrameArgs CaptureArgs{ GraphBuilder };
+					CaptureArgs.MediaCapture = InMediaCapture;
+					CaptureArgs.DesiredSize = InDesiredSize;
+					CaptureArgs.ResourceToCapture = MoveTemp(SourceTexture);
+
+					InMediaCapture->ProcessCapture_RenderThread(CapturingFrame, CaptureArgs);
+					UE_LOG(LogMediaIOCore, Verbose, TEXT("[%s] - Captured frame %d"), *InMediaCapture->MediaOutputName, CapturingFrame ? CapturingFrame->FrameId : -1);
+					
+					GraphBuilder.Execute();
+				}
+				
+				//Process the next pending frame
+				
+				UE_LOG(LogMediaIOCore, Verbose, TEXT("[%s] - Processing pending frame %d"), *InMediaCapture->MediaOutputName, NextPending ? NextPending->FrameId : -1);
+				if (NextPending)
+				{
+					InMediaCapture->ProcessReadyFrame_RenderThread(RHICmdList, InMediaCapture, NextPending);
+				}
+
+				// Whatever happens, we want to decrement our counter to track enqueued commands
+				--InMediaCapture->WaitingForRenderCommandExecutionCounter;
+			});
+
+			//If auto-stopping, count the number of frame captures requested and stop when reaching 0.
+			if (DesiredCaptureOptions.bAutostopOnCapture && GetState() == EMediaCaptureState::Capturing && --DesiredCaptureOptions.NumberOfFramesToCapture <= 0)
+			{
+				StopCapture(true);
+			}
 		}
-		InViewport->SetFixedViewportSize(0, 0);
-		bViewportHasFixedViewportSize = false;
 	}
 }
 
@@ -1456,9 +1767,10 @@ void UMediaCapture::InitializeOutputResources(int32 InNumberOfBuffers)
 	}
 
 	// Recreate frame manager which can trigger cleaning up its captured frames if it exists
-	FrameManager = MakePimpl<UE::MediaCaptureData::FFrameManager>();
+	FrameManager = MakePimpl<FFrameManager>();
 
 	bOutputResourcesInitialized = false;
+	bSyncHandlersInitialized = false;
 	NumberOfCaptureFrame = InNumberOfBuffers;
 
 	check(NumberOfCaptureFrame >= 2);
@@ -1471,7 +1783,7 @@ void UMediaCapture::InitializeOutputResources(int32 InNumberOfBuffers)
 			{
 				if (This->DesiredOutputResourceType == EMediaCaptureResourceType::Texture)
 				{
-					TUniquePtr<FTextureCaptureFrame> NewFrame = MakeUnique<FTextureCaptureFrame>(Index);
+					TSharedPtr<FTextureCaptureFrame> NewFrame = MakeShared<FTextureCaptureFrame>(Index);
 					FRDGTextureDesc OutputDesc = FRDGTextureDesc::Create2D(
 						This->DesiredOutputSize,
 						This->DesiredOutputPixelFormat,
@@ -1491,7 +1803,7 @@ void UMediaCapture::InitializeOutputResources(int32 InNumberOfBuffers)
 				}
 				else
 				{
-					TUniquePtr<FBufferCaptureFrame> NewFrame = MakeUnique<FBufferCaptureFrame>(Index);
+					TSharedPtr<FBufferCaptureFrame> NewFrame = MakeShared<FBufferCaptureFrame>(Index);
 
 					if (This->DesiredOutputBufferDescription.NumElements > 0)
 					{
@@ -1516,6 +1828,51 @@ void UMediaCapture::InitializeOutputResources(int32 InNumberOfBuffers)
 		});
 }
 
+void UMediaCapture::InitializeSyncHandlers_RenderThread()
+{
+	SyncHandlers.Reset(NumberOfCaptureFrame);
+	for (int32 Index = 0; Index < NumberOfCaptureFrame; ++Index)
+	{
+		TSharedPtr<FMediaCaptureSyncData> SyncData = MakeShared<FMediaCaptureSyncData>();
+		SyncData->RHIFence = RHICreateGPUFence(*FString::Printf(TEXT("MediaCaptureSync_%02d"), Index));
+		SyncHandlers.Add(MoveTemp(SyncData));
+	}
+
+	bSyncHandlersInitialized = true;
+}
+
+void UMediaCapture::InitializeCaptureFrame(const TSharedPtr<UE::MediaCaptureData::FCaptureFrame>& CaptureFrame)
+{
+	if (CaptureFrame)
+	{
+		CaptureFrame->CaptureBaseData.SourceFrameTimecode = FApp::GetTimecode();
+		CaptureFrame->CaptureBaseData.SourceFrameTimecodeFramerate = FApp::GetTimecodeFrameRate();
+		CaptureFrame->CaptureBaseData.SourceFrameNumberRenderThread = GFrameNumber;
+		CaptureFrame->CaptureBaseData.SourceFrameNumber = ++CaptureRequestCount;
+		CaptureFrame->UserData = GetCaptureFrameUserData_GameThread();
+	}
+}
+
+TSharedPtr<UMediaCapture::FMediaCaptureSyncData> UMediaCapture::GetAvailableSyncHandler() const
+{
+	const auto FindAvailableHandlerFunc = [](const TSharedPtr<FMediaCaptureSyncData>& Item)
+	{
+		if (Item->bIsBusy == false)
+		{
+			return true;
+		}
+
+		return false;
+	};
+
+	if (const TSharedPtr<FMediaCaptureSyncData>* FoundItem = SyncHandlers.FindByPredicate(FindAvailableHandlerFunc))
+	{
+		return *FoundItem;
+	}
+
+	return nullptr;
+}
+
 bool UMediaCapture::ValidateMediaOutput() const
 {
 	if (MediaOutput == nullptr)
@@ -1552,8 +1909,20 @@ bool UMediaCapture::CaptureRenderTargetImpl(UTextureRenderTarget2D * InRenderTar
 
 void UMediaCapture::OnBeginFrame_GameThread()
 {
+	TRACE_CPUPROFILER_EVENT_SCOPE(MediaCapture::BeginFrame);
+
+	if (ViewExtension && !ViewExtension->IsValid())
+	{
+		SetState(EMediaCaptureState::Error);
+	}
+
+	if (GetState() == EMediaCaptureState::Error)
+	{
+		StopCapture(false);
+		return;
+	}
 
-	if (CaptureSourceType == EMediaCaptureSourceType::RHI_RESOURCE && (GetState() == EMediaCaptureState::Capturing))
+	if (CaptureSource->GetCaptureType() == UE::MediaCapture::Private::ECaptureType::Immediate && (GetState() == EMediaCaptureState::Capturing))
 	{
 		// Queue capture data to be consumed when capture requests are done on render thread
 		FQueuedCaptureData CaptureData;
@@ -1561,6 +1930,7 @@ void UMediaCapture::OnBeginFrame_GameThread()
 		CaptureData.BaseData.SourceFrameTimecodeFramerate = FApp::GetTimecodeFrameRate();
 		CaptureData.BaseData.SourceFrameNumberRenderThread = GFrameCounter;
 		CaptureData.UserData = GetCaptureFrameUserData_GameThread();
+		
 		CaptureDataQueue.Enqueue(MoveTemp(CaptureData));
 	}
 }
@@ -1569,7 +1939,7 @@ void UMediaCapture::OnEndFrame_GameThread()
 {
 	using namespace UE::MediaCaptureData;
 
-	TRACE_CPUPROFILER_EVENT_SCOPE(MediaCapture::OnEndFrame_GameThread);
+	TRACE_CPUPROFILER_EVENT_SCOPE_TEXT(*FString::Printf(TEXT("MediaCapture End Frame %d"), GFrameNumber));
 
 	if (!bOutputResourcesInitialized)
 	{
@@ -1603,155 +1973,67 @@ void UMediaCapture::OnEndFrame_GameThread()
 		return;
 	}
 
-	FCaptureFrame* CapturingFrame = nullptr;
-	if(GetState() != EMediaCaptureState::StopRequested)
 	{
-		CapturingFrame = FrameManager->GetNextAvailable();
-	}
+		FScopeLock ScopedLock(&PendingReadbackTasksCriticalSection);
 
-	if (CapturingFrame == nullptr && GetState() != EMediaCaptureState::StopRequested)
-	{
-		if (DesiredCaptureOptions.OverrunAction == EMediaCaptureOverrunAction::Flush)
+		// Clean up completed tasks.
+		for (auto It = PendingReadbackTasks.CreateIterator(); It; ++It)
 		{
-			TRACE_CPUPROFILER_EVENT_SCOPE(MediaCapture::FlushRenderingCommands);
-			UE_LOG(LogMediaIOCore, VeryVerbose, TEXT("[%s] - Flushing commands."), *MediaOutputName);
-			FlushRenderingCommands();
-
-			// After flushing, we should have access to an available frame, if not, it's not expected.
-			CapturingFrame = FrameManager->GetNextAvailable();
-			if (CapturingFrame == nullptr)
+			if (It->IsCompleted())
 			{
-				UE_LOG(LogMediaIOCore, Error, TEXT("[%s] - Flushing commands didn't give us back available frames to process.") , *MediaOutputName);
-				StopCapture(false);
-				return;
+				It.RemoveCurrent();
 			}
 		}
-		else
-		{
-			// Selected options is to skip capturing a frame if overrun happens
-			UE_LOG(LogMediaIOCore, Verbose, TEXT("[%s] - No frames available for capture. Skipping"), *MediaOutputName);
-		}
-	}
-
-	if (CapturingFrame)
-	{
-		CapturingFrame->CaptureBaseData.SourceFrameTimecode = FApp::GetTimecode();
-		CapturingFrame->CaptureBaseData.SourceFrameTimecodeFramerate = FApp::GetTimecodeFrameRate();
-		CapturingFrame->CaptureBaseData.SourceFrameNumberRenderThread = GFrameCounter;
-		CapturingFrame->CaptureBaseData.SourceFrameNumber = ++CaptureRequestCount;
-		CapturingFrame->UserData = GetCaptureFrameUserData_GameThread();
 	}
 
-	PrintFrameState();
-	
-	// Init variables for ENQUEUE_RENDER_COMMAND.
-	//The Lock only synchronize while we are copying the value to the enqueue. The viewport and the rendertarget may change while we are in the enqueue command.
-	{
-		FScopeLock Lock(&AccessingCapturingSource);
-
-		TSharedPtr<FSceneViewport> CapturingSceneViewportPin = CapturingSceneViewport.Pin();
-		FSceneViewport* InCapturingSceneViewport = CapturingSceneViewportPin.Get();
-		FTextureRenderTargetResource* InTextureRenderTargetResource = CapturingRenderTarget ? CapturingRenderTarget->GameThread_GetRenderTargetResource() : nullptr;
-		FIntPoint InDesiredSize = DesiredSize;
-		UMediaCapture* InMediaCapture = this;
-
-
-		if (InCapturingSceneViewport != nullptr || InTextureRenderTargetResource != nullptr)
-		{
-			++WaitingForRenderCommandExecutionCounter;
-
-			const FRHIGPUMask SourceGPUMask = ValidSourceGPUMask;
-
-			UE_LOG(LogMediaIOCore, Verbose, TEXT("[%s] - Queuing frame to capture %d"), *InMediaCapture->MediaOutputName, CapturingFrame ? CapturingFrame->FrameId : -1);
-
-			// RenderCommand to be executed on the RenderThread
-			ENQUEUE_RENDER_COMMAND(FMediaOutputCaptureFrameCreateTexture)(
-				[InMediaCapture, CapturingFrame, InCapturingSceneViewport, InTextureRenderTargetResource, InDesiredSize, SourceGPUMask](FRHICommandListImmediate& RHICmdList)
-			{
-				SCOPED_GPU_MASK(RHICmdList, SourceGPUMask);
-
-				// Peek next pending before capturing to avoid getting the newly captured one. 
-				// It shouldn't happen but if it happens it will start waiting for the capture to complete
-				FCaptureFrame* NextPending = InMediaCapture->FrameManager->PeekNextPending();
-				if (NextPending && CapturingFrame)
-				{
-					ensure(NextPending->FrameId != CapturingFrame->FrameId);
-				}
-
-				// Capture frame
-				{
-					FRDGBuilder GraphBuilder(RHICmdList);
-
-					FTexture2DRHIRef SourceTexture = FMediaCaptureHelper::GetSourceTextureForInput(RHICmdList, InCapturingSceneViewport, InTextureRenderTargetResource);
-					InMediaCapture->ProcessCapture_RenderThread(GraphBuilder, InMediaCapture, CapturingFrame, MoveTemp(SourceTexture), InDesiredSize);
-					UE_LOG(LogMediaIOCore, Verbose, TEXT("[%s] - Captured frame %d"), *InMediaCapture->MediaOutputName, CapturingFrame ? CapturingFrame->FrameId : -1);
-					
-					GraphBuilder.Execute();
-				}
-				
-				//Process pending frame
-				UE_LOG(LogMediaIOCore, Verbose, TEXT("[%s] - Processing pending frame %d"), *InMediaCapture->MediaOutputName, NextPending ? NextPending->FrameId : -1);
-				if (NextPending)
-				{
-					InMediaCapture->ProcessReadyFrame_RenderThread(RHICmdList, InMediaCapture, NextPending);
-				}
-				
-
-				// Whatever happens, we want to decrement our counter to track enqueued commands
-				--InMediaCapture->WaitingForRenderCommandExecutionCounter;
-			});
-
-			//If auto-stopping, count the number of frame captures requested and stop when reaching 0.
-			if(DesiredCaptureOptions.bAutostopOnCapture && GetState() == EMediaCaptureState::Capturing && --DesiredCaptureOptions.NumberOfFramesToCapture <= 0)
-			{
-				StopCapture(true);
-			}
-		}
-	}
+	ProcessCapture_GameThread();
 }
 
-bool UMediaCapture::ProcessCapture_RenderThread(FRDGBuilder& GraphBuilder, UMediaCapture* InMediaCapture, UE::MediaCaptureData::FCaptureFrame* CapturingFrame, FTexture2DRHIRef InResourceToCapture,	FIntPoint InDesiredSize)
+bool UMediaCapture::ProcessCapture_RenderThread(const TSharedPtr<UE::MediaCaptureData::FCaptureFrame>& CapturingFrame, const UE::MediaCaptureData::FCaptureFrameArgs& Args)
 {
 	using namespace UE::MediaCaptureData;
 	
-	RDG_GPU_STAT_SCOPE(GraphBuilder, MediaCapture_ProcessCapture)
+	RDG_GPU_STAT_SCOPE(Args.GraphBuilder, MediaCapture_ProcessCapture)
 	TRACE_CPUPROFILER_EVENT_SCOPE(UMediaCapture::ProcessCapture_RenderThread);
+	int FrameNumber = -1;
+	if (CapturingFrame)
+	{
+		FrameNumber = CapturingFrame->CaptureBaseData.SourceFrameNumberRenderThread;
+	}
+	TRACE_CPUPROFILER_EVENT_SCOPE_TEXT(*FString::Printf(TEXT("Process Capture Render Thread Frame %d"), FrameNumber));
 
 	if (CapturingFrame)
 	{
 		// Call the capture frame algo based on the specific type of resource we are using
 		bool bHasCaptureSuceeded = false;
-		FCaptureFrameArgs CaptureArgs = { GraphBuilder, InMediaCapture, MoveTemp(InResourceToCapture), InDesiredSize};
 		if (DesiredOutputResourceType == EMediaCaptureResourceType::Texture)
 		{
 			if (ensure(CapturingFrame->IsTextureResource()))
 			{
-				FTextureCaptureFrame* TextureFrame = static_cast<FTextureCaptureFrame*>(CapturingFrame);
-				InMediaCapture->BeforeFrameCaptured_RenderingThread(CapturingFrame->CaptureBaseData, CapturingFrame->UserData, TextureFrame->GetTextureResource());
-				bHasCaptureSuceeded = FMediaCaptureHelper::CaptureFrame(CaptureArgs, TextureFrame);
+				Args.MediaCapture->BeforeFrameCaptured_RenderingThread(CapturingFrame->CaptureBaseData, CapturingFrame->UserData, CapturingFrame->GetTextureResource());
+				bHasCaptureSuceeded = FMediaCaptureHelper::CaptureFrame(Args, CapturingFrame);
 			}
 			else
 			{
-				UE_LOG(LogMediaIOCore, Error, TEXT("The capture will stop for '%s'. Capture frame was expected to use Texture resource but wasn't."), *InMediaCapture->MediaOutputName);
+				UE_LOG(LogMediaIOCore, Error, TEXT("The capture will stop for '%s'. Capture frame was expected to use Texture resource but wasn't."), *Args.MediaCapture->MediaOutputName);
 			}
 		}
 		else
 		{
 			if (ensure(CapturingFrame->IsBufferResource()))
 			{
-				FBufferCaptureFrame* BufferFrame = static_cast<FBufferCaptureFrame*>(CapturingFrame);
-				InMediaCapture->BeforeFrameCaptured_RenderingThread(CapturingFrame->CaptureBaseData, CapturingFrame->UserData, BufferFrame->GetBufferResource());
-				bHasCaptureSuceeded = FMediaCaptureHelper::CaptureFrame(CaptureArgs, BufferFrame);
+				Args.MediaCapture->BeforeFrameCaptured_RenderingThread(CapturingFrame->CaptureBaseData, CapturingFrame->UserData, CapturingFrame->GetBufferResource());
+				bHasCaptureSuceeded = FMediaCaptureHelper::CaptureFrame(Args, CapturingFrame);
 			}
 			else
 			{
-				UE_LOG(LogMediaIOCore, Error, TEXT("The capture will stop for '%s'. Capture frame was expected to use Buffer resource but wasn't."), *InMediaCapture->MediaOutputName);
+				UE_LOG(LogMediaIOCore, Error, TEXT("The capture will stop for '%s'. Capture frame was expected to use Buffer resource but wasn't."), *Args.MediaCapture->MediaOutputName);
 			}
 		}
 		
 		if (bHasCaptureSuceeded == false)
 		{
-			InMediaCapture->SetState(EMediaCaptureState::Error);
+			Args.MediaCapture->SetState(EMediaCaptureState::Error);
 		}
 
 		return bHasCaptureSuceeded;
@@ -1760,10 +2042,9 @@ bool UMediaCapture::ProcessCapture_RenderThread(FRDGBuilder& GraphBuilder, UMedi
 	return false;
 }
 
-bool UMediaCapture::ProcessReadyFrame_RenderThread(FRHICommandListImmediate& RHICmdList, UMediaCapture* InMediaCapture, UE::MediaCaptureData::FCaptureFrame* ReadyFrame)
+bool UMediaCapture::ProcessReadyFrame_RenderThread(FRHICommandListImmediate& RHICmdList, UMediaCapture* InMediaCapture, const TSharedPtr<UE::MediaCaptureData::FCaptureFrame>& ReadyFrame)
 {
 	using namespace UE::MediaCaptureData;
-
 	TRACE_CPUPROFILER_EVENT_SCOPE(UMediaCapture::ProcessReadyFrame_RenderThread);
 
 	bool bWasFrameProcessed = true;
@@ -1801,10 +2082,11 @@ bool UMediaCapture::ProcessReadyFrame_RenderThread(FRHICommandListImmediate& RHI
 				bWasFrameProcessed = false;
 			}
 
-			if(ColorDataBuffer)
+			if (ColorDataBuffer)
 			{
 				{
-					SCOPE_CYCLE_COUNTER(STAT_MediaCapture_RenderThread_CaptureCallback)
+					SCOPE_CYCLE_COUNTER(STAT_MediaCapture_RenderThread_RHI_CaptureCallback)
+					TRACE_CPUPROFILER_EVENT_SCOPE_TEXT(*FString::Printf(TEXT("MediaCapture Output Frame %d"), ReadyFrame->CaptureBaseData.SourceFrameNumberRenderThread));
 
 					// The Width/Height of the surface may be different then the DesiredOutputSize : Some underlying implementations enforce a specific stride, therefore
 					// there may be padding at the end of each row.
@@ -1816,7 +2098,7 @@ bool UMediaCapture::ProcessReadyFrame_RenderThread(FRHICommandListImmediate& RHI
 
 			if (bIsReadbackReady)
 			{
-				UE_LOG(LogMediaIOCore, Verbose, TEXT("[%s] - Completed pending frame %d."), *InMediaCapture->MediaOutputName, ReadyFrame->FrameId);
+				UE_LOG(LogMediaIOCore, Verbose, TEXT("[%s - %s] - Completed pending frame %d."), *InMediaCapture->MediaOutputName, *FThreadManager::GetThreadName(FPlatformTLS::GetCurrentThreadId()), ReadyFrame->FrameId);
 				ReadyFrame->bReadbackRequested = false;
 				--PendingFrameCount;
 				InMediaCapture->FrameManager->CompleteNextPending(*ReadyFrame);
@@ -1826,24 +2108,23 @@ bool UMediaCapture::ProcessReadyFrame_RenderThread(FRHICommandListImmediate& RHI
 		{
 			if (ReadyFrame->IsTextureResource())
 			{
-				FTextureCaptureFrame* TextureFrame = static_cast<FTextureCaptureFrame*>(ReadyFrame);
 				{
 					TRACE_CPUPROFILER_EVENT_SCOPE(UMediaCapture::UnlockDMATexture_RenderThread);
-					UnlockDMATexture_RenderThread(TextureFrame->RenderTarget->GetRHI());
+					UnlockDMATexture_RenderThread(ReadyFrame->GetTextureResource());
 				}
 
 				TRACE_CPUPROFILER_EVENT_SCOPE(UMediaCapture::RHIResourceCaptured);
+				TRACE_CPUPROFILER_EVENT_SCOPE_TEXT(*FString::Printf(TEXT("MediaCapture Output Frame %d"), ReadyFrame->CaptureBaseData.SourceFrameNumberRenderThread));
 				SCOPE_CYCLE_COUNTER(STAT_MediaCapture_RenderThread_CaptureCallback)
-					InMediaCapture->OnRHIResourceCaptured_RenderingThread(ReadyFrame->CaptureBaseData, ReadyFrame->UserData, TextureFrame->GetRHIResource());
+					InMediaCapture->OnRHIResourceCaptured_RenderingThread(ReadyFrame->CaptureBaseData, ReadyFrame->UserData, ReadyFrame->GetTextureResource());
 			}
 			else
 			{
 				TRACE_CPUPROFILER_EVENT_SCOPE(UMediaCapture::RHIResourceCaptured);
-				FBufferCaptureFrame* BufferFrame = static_cast<FBufferCaptureFrame*>(ReadyFrame);
-				InMediaCapture->OnRHIResourceCaptured_RenderingThread(ReadyFrame->CaptureBaseData, ReadyFrame->UserData, BufferFrame->GetRHIResource());
+				InMediaCapture->OnRHIResourceCaptured_RenderingThread(ReadyFrame->CaptureBaseData, ReadyFrame->UserData, ReadyFrame->GetBufferResource());
 			}
 
-			UE_LOG(LogMediaIOCore, Verbose, TEXT("[%s] - Completed pending frame %d."), *InMediaCapture->MediaOutputName, ReadyFrame->FrameId);
+			UE_LOG(LogMediaIOCore, Verbose, TEXT("[%s - %s] - Completed pending frame %d."), *InMediaCapture->MediaOutputName, *FThreadManager::GetThreadName(FPlatformTLS::GetCurrentThreadId()), ReadyFrame->FrameId);
 			ReadyFrame->bDoingGPUCopy = false;
 			InMediaCapture->FrameManager->CompleteNextPending(*ReadyFrame);
 			--PendingFrameCount;
@@ -1855,16 +2136,7 @@ bool UMediaCapture::ProcessReadyFrame_RenderThread(FRHICommandListImmediate& RHI
 
 void UMediaCapture::PrintFrameState()
 {
-	using namespace UE::MediaCaptureData;
-
-	TStringBuilder<256> ReadbackInfoBuilder;
-	ReadbackInfoBuilder << "\n";
-	TConstArrayView<TUniquePtr<FCaptureFrame>> Frames = FrameManager->GetFrames();
-	for (int32 Index = 0; Index < Frames.Num(); Index++)
-	{
-		ReadbackInfoBuilder << FString::Format(TEXT("Frame {0} readback requested: {1}\n"), { Index, Frames[Index]->IsPending() ? TEXT("true") : TEXT("false") });
-	}
-	UE_LOG(LogMediaIOCore, VeryVerbose, TEXT("%s"), ReadbackInfoBuilder.GetData());
+	UE_LOG(LogMediaIOCore, VeryVerbose, TEXT("%s"), *FrameManager->GetFramesState());
 }
 
 /* namespace MediaCaptureDetails implementation
@@ -1918,126 +2190,9 @@ namespace MediaCaptureDetails
 		return (OutSceneViewport.IsValid());
 	}
 
-	bool ValidateSize(const FIntPoint TargetSize, const FIntPoint& DesiredSize, const FMediaCaptureOptions& CaptureOptions, const bool bCurrentlyCapturing)
-	{
-		if (CaptureOptions.Crop == EMediaCaptureCroppingType::None)
-		{
-			if (DesiredSize.X != TargetSize.X || DesiredSize.Y != TargetSize.Y)
-			{
-				UE_LOG(LogMediaIOCore, Error, TEXT("Can not %s the capture. The Render Target size doesn't match with the requested size. SceneViewport: %d,%d  MediaOutput: %d,%d")
-					, bCurrentlyCapturing ? TEXT("continue") : TEXT("start")
-					, TargetSize.X, TargetSize.Y
-					, DesiredSize.X, DesiredSize.Y);
-				return false;
-			}
-		}
-		else
-		{
-			FIntPoint StartCapturePoint = FIntPoint::ZeroValue;
-			if (CaptureOptions.Crop == EMediaCaptureCroppingType::Custom)
-			{
-				if (CaptureOptions.CustomCapturePoint.X < 0 || CaptureOptions.CustomCapturePoint.Y < 0)
-				{
-					UE_LOG(LogMediaIOCore, Error, TEXT("Can not %s the capture. The start capture point is negatif. Start Point: %d,%d")
-						, bCurrentlyCapturing ? TEXT("continue") : TEXT("start")
-						, StartCapturePoint.X, StartCapturePoint.Y);
-					return false;
-				}
-				StartCapturePoint = CaptureOptions.CustomCapturePoint;
-			}
-
-			if (DesiredSize.X + StartCapturePoint.X > TargetSize.X || DesiredSize.Y + StartCapturePoint.Y > TargetSize.Y)
-			{
-				UE_LOG(LogMediaIOCore, Error, TEXT("Can not %s the capture. The Render Target size is too small for the requested cropping options. SceneViewport: %d,%d  MediaOutput: %d,%d Start Point: %d,%d")
-					, bCurrentlyCapturing ? TEXT("continue") : TEXT("start")
-					, TargetSize.X, TargetSize.Y
-					, DesiredSize.X, DesiredSize.Y
-					, StartCapturePoint.X, StartCapturePoint.Y);
-				return false;
-			}
-		}
-
-		return true;
-	}
-
-	bool ValidateSceneViewport(const TSharedPtr<FSceneViewport>& SceneViewport, const FMediaCaptureOptions& CaptureOptions, const FIntPoint& DesiredSize, const EPixelFormat DesiredPixelFormat, const bool bCurrentlyCapturing)
-	{
-		if (!SceneViewport.IsValid())
-		{
-			UE_LOG(LogMediaIOCore, Error, TEXT("Can not %s the capture. The Scene Viewport is invalid.")
-				, bCurrentlyCapturing ? TEXT("continue") : TEXT("start"));
-			return false;
-		}
-
-		const FIntPoint SceneViewportSize = SceneViewport->GetRenderTargetTextureSizeXY();
-		if (!ValidateSize(SceneViewportSize, DesiredSize, CaptureOptions, bCurrentlyCapturing))
-		{
-			return false;
-		}
-
-		static const auto CVarDefaultBackBufferPixelFormat = IConsoleManager::Get().FindTConsoleVariableDataInt(TEXT("r.DefaultBackBufferPixelFormat"));
-		EPixelFormat SceneTargetFormat = EDefaultBackBufferPixelFormat::Convert2PixelFormat(EDefaultBackBufferPixelFormat::FromInt(CVarDefaultBackBufferPixelFormat->GetValueOnGameThread()));
-		if (DesiredPixelFormat != SceneTargetFormat)
-		{
-			if (!SupportedRgbaSwizzleFormats.Contains(SceneTargetFormat) || !CaptureOptions.bConvertToDesiredPixelFormat)
-			{
-				UE_LOG(LogMediaIOCore, Error, TEXT("Can not %s the capture. The Render Target pixel format doesn't match with the requested pixel format. %sRenderTarget: %s MediaOutput: %s")
-					, bCurrentlyCapturing ? TEXT("continue") : TEXT("start")
-					, (SupportedRgbaSwizzleFormats.Contains(SceneTargetFormat) && !CaptureOptions.bConvertToDesiredPixelFormat) ? TEXT("Please enable \"Convert To Desired Pixel Format\" option in Media Capture settings") : TEXT("")
-					, GetPixelFormatString(SceneTargetFormat)
-					, GetPixelFormatString(DesiredPixelFormat));
-				return false;
-			}
-			else
-			{
-				UE_LOG(LogMediaIOCore, Warning, TEXT("The Render Target pixel format doesn't match with the requested pixel format. Render target will be automatically converted. This could have a slight performance impact. RenderTarget: %s MediaOutput: %s")
-					, GetPixelFormatString(SceneTargetFormat)
-					, GetPixelFormatString(DesiredPixelFormat));
-			}
-		}
-
-		return true;
-	}
-
-	bool ValidateTextureRenderTarget2D(const UTextureRenderTarget2D* InRenderTarget2D, const FMediaCaptureOptions& CaptureOptions, const FIntPoint& DesiredSize, const EPixelFormat DesiredPixelFormat, const bool bCurrentlyCapturing)
-	{
-		if (InRenderTarget2D == nullptr)
-		{
-			UE_LOG(LogMediaIOCore, Error, TEXT("Couldn't %s the capture. The Render Target is invalid.")
-				, bCurrentlyCapturing ? TEXT("continue") : TEXT("start"));
-			return false;
-		}
-
-		if (!ValidateSize(FIntPoint(InRenderTarget2D->SizeX, InRenderTarget2D->SizeY), DesiredSize, CaptureOptions, bCurrentlyCapturing))
-		{
-			return false;
-		}
-
-		if (DesiredPixelFormat != InRenderTarget2D->GetFormat())
-		{
-			if (!SupportedRgbaSwizzleFormats.Contains(InRenderTarget2D->GetFormat()) || !CaptureOptions.bConvertToDesiredPixelFormat)
-			{
-				UE_LOG(LogMediaIOCore, Error, TEXT("Can not %s the capture. The Render Target pixel format doesn't match with the requested pixel format. %sRenderTarget: %s MediaOutput: %s")
-					, bCurrentlyCapturing ? TEXT("continue") : TEXT("start")
-					, (SupportedRgbaSwizzleFormats.Contains(InRenderTarget2D->GetFormat()) && !CaptureOptions.bConvertToDesiredPixelFormat) ? TEXT("Please enable \"Convert To Desired Pixel Format\" option in Media Capture settings. ") : TEXT("")
-					, GetPixelFormatString(InRenderTarget2D->GetFormat())
-					, GetPixelFormatString(DesiredPixelFormat));
-				return false;
-			}
-			else
-			{
-				UE_LOG(LogMediaIOCore, Warning, TEXT("The Render Target pixel format doesn't match with the requested pixel format. Render target will be automatically converted. This could have a slight performance impact. RenderTarget: %s MediaOutput: %s")
-					, GetPixelFormatString(InRenderTarget2D->GetFormat())
-					, GetPixelFormatString(DesiredPixelFormat));
-			}
-		}
-
-		return true;
-	}
-
-	bool ValidateIsCapturing(const UMediaCapture& CaptureToBeValidated)
+	bool ValidateIsCapturing(const UMediaCapture* CaptureToBeValidated)
 	{
-		if (CaptureToBeValidated.GetState() != EMediaCaptureState::Capturing && CaptureToBeValidated.GetState() != EMediaCaptureState::Preparing)
+		if (CaptureToBeValidated->GetState() != EMediaCaptureState::Capturing && CaptureToBeValidated->GetState() != EMediaCaptureState::Preparing)
 		{
 			UE_LOG(LogMediaIOCore, Error, TEXT("Can not update the capture. There is no capture currently.\
 			Only use UpdateSceneViewport or UpdateTextureRenderTarget2D when the state is Capturing or Preparing"));
diff --git a/Engine/Plugins/Media/MediaIOFramework/Source/MediaIOCore/Private/MediaCaptureSceneViewExtension.h b/Engine/Plugins/Media/MediaIOFramework/Source/MediaIOCore/Private/MediaCaptureSceneViewExtension.h
new file mode 100644
--- /dev/null
+++ b/Engine/Plugins/Media/MediaIOFramework/Source/MediaIOCore/Private/MediaCaptureSceneViewExtension.h
@@ -0,0 +1,154 @@
+// Copyright Epic Games, Inc. All Rights Reserved.
+
+#pragma once
+
+#include "SceneViewExtension.h"
+#include "MediaCapture.h"
+#include "MediaIOCoreModule.h"
+#include "ImagePixelData.h"
+#include "PostProcess/PostProcessing.h"
+#include "PostProcess/PostProcessMaterial.h"
+#include "RenderGraphUtils.h"
+#include "RHI.h"
+#include "RHICommandList.h"
+#include "RHIResources.h"
+#include "SceneView.h"
+#include "ScreenPass.h"
+
+class FRDGTexture;
+
+/**
+ * View extension that calls into MediaCapture to capture frames as soon as the frame is rendered.
+ */
+class FMediaCaptureSceneViewExtension : public FSceneViewExtensionBase
+{
+public:
+	
+	FMediaCaptureSceneViewExtension(const FAutoRegister& InAutoRegister, UMediaCapture* InMediaCapture, EMediaCapturePhase InCapturePhase)
+		: FSceneViewExtensionBase(InAutoRegister)
+		, WeakCapture(InMediaCapture)
+		, CapturePhase(InCapturePhase)
+	{
+	}
+
+	//~ Begin FSceneViewExtensionBase Interface
+	virtual void SetupViewFamily(FSceneViewFamily& InViewFamily) override {};
+	virtual void SetupView(FSceneViewFamily& InViewFamily, FSceneView& InView) override 
+	{
+		if (bValidPhase)
+		{
+			// Copied from PostProcessing.h
+			if (InView.GetFeatureLevel() >= ERHIFeatureLevel::SM5)
+			{
+				bPostProcessingEnabled =
+					InView.Family->EngineShowFlags.PostProcessing &&
+					!InView.Family->EngineShowFlags.VisualizeDistanceFieldAO &&
+					!InView.Family->EngineShowFlags.VisualizeShadingModels &&
+					!InView.Family->EngineShowFlags.VisualizeVolumetricCloudConservativeDensity &&
+					!InView.Family->EngineShowFlags.VisualizeVolumetricCloudEmptySpaceSkipping &&
+					!InView.Family->EngineShowFlags.ShaderComplexity;
+			}
+			else
+			{
+				bPostProcessingEnabled = InView.Family->EngineShowFlags.PostProcessing && !InView.Family->EngineShowFlags.ShaderComplexity && IsMobileHDR();
+			}
+
+			if (CapturePhase != EMediaCapturePhase::BeforePostProcessing && CapturePhase != EMediaCapturePhase::EndFrame)
+			{
+				if (!bPostProcessingEnabled)
+				{
+					LastErrorMessage = TEXT("Media Capture will not work since it is scheduled in a post processing phase and post processing is not enabled.");
+					UE_LOG(LogMediaIOCore, Warning, TEXT("%s"), *LastErrorMessage);
+					bValidPhase = false;
+				}
+			}
+		}
+	};
+	virtual void BeginRenderViewFamily(FSceneViewFamily& InViewFamily) override {}; 
+
+	virtual void PrePostProcessPass_RenderThread(FRDGBuilder& GraphBuilder, const FSceneView& View, const FPostProcessingInputs& Inputs) override
+	{
+		if (CapturePhase == EMediaCapturePhase::BeforePostProcessing)
+		{
+			Inputs.Validate();
+
+			if (FRDGTextureRef TextureRef = Inputs.ViewFamilyTexture)
+			{
+				RDG_EVENT_SCOPE(GraphBuilder, "MediaCaptureSceneExtension");
+				if (UMediaCapture* MediaCapture = WeakCapture.Get())
+				{
+					MediaCapture->CaptureImmediate_RenderThread(GraphBuilder, TextureRef);
+				}
+			}
+		}
+	}
+
+	virtual void SubscribeToPostProcessingPass(EPostProcessingPass PassId, FAfterPassCallbackDelegateArray& InOutPassCallbacks, bool bIsPassEnabled) override
+	{		
+		if ((CapturePhase == EMediaCapturePhase::AfterMotionBlur && PassId == EPostProcessingPass::MotionBlur)
+			|| (CapturePhase == EMediaCapturePhase::AfterToneMap && PassId == EPostProcessingPass::Tonemap)
+			|| (CapturePhase == EMediaCapturePhase::AfterFXAA && PassId == EPostProcessingPass::FXAA))
+		{
+			if (!bIsPassEnabled)
+			{
+				UE_LOG(LogMediaIOCore, Warning, TEXT("Media Capture will not work since the selected capture phase is not enabled."));
+			}
+			InOutPassCallbacks.Add(FAfterPassCallbackDelegate::CreateSP(this, &FMediaCaptureSceneViewExtension::PostProcessCallback_RenderThread));
+		}
+	}
+
+	virtual bool IsActiveThisFrame_Internal(const FSceneViewExtensionContext& Context) const override
+	{
+		return true;
+	}
+
+	FScreenPassTexture PostProcessCallback_RenderThread(FRDGBuilder& GraphBuilder, const FSceneView& View, const FPostProcessMaterialInputs& InOutInputs)
+	{
+		TRACE_CPUPROFILER_EVENT_SCOPE(MediaCaptureExtensionCallback);
+
+		FScreenPassTexture SceneColor = InOutInputs.GetInput(EPostProcessMaterialInput::SceneColor);
+		InOutInputs.Validate();
+		
+		if (FRDGTextureRef TextureRef = SceneColor.Texture)
+		{
+			RDG_EVENT_SCOPE(GraphBuilder, "MediaCaptureSceneExtension");
+			if (WeakCapture.IsValid())
+			{
+				WeakCapture->CaptureImmediate_RenderThread(GraphBuilder, TextureRef);
+			}
+		}
+
+		return ReturnUntouchedSceneColorForPostProcessing(InOutInputs);
+	}
+	
+	/** 
+	 * A helper function that extracts the right scene color texture, untouched, to be used further in post processing. 
+	 */
+	FScreenPassTexture ReturnUntouchedSceneColorForPostProcessing(const FPostProcessMaterialInputs& InOutInputs)
+	{
+		if (InOutInputs.OverrideOutput.IsValid())
+		{
+			return InOutInputs.OverrideOutput;
+		}
+		else
+		{
+			/** We don't want to modify scene texture in any way. We just want it to be passed back onto the next stage. */
+			FScreenPassTexture SceneTexture = const_cast<FScreenPassTexture&>(InOutInputs.Textures[(uint32)EPostProcessMaterialInput::SceneColor]);
+			return SceneTexture;
+		}
+	}
+
+	bool IsValid() const
+	{
+		return bValidPhase;
+	}
+
+private:
+	TWeakObjectPtr<UMediaCapture> WeakCapture;
+	EMediaCapturePhase CapturePhase = EMediaCapturePhase::AfterMotionBlur;
+	bool bPostProcessingEnabled = true;
+	bool bValidPhase = true;
+	FString LastErrorMessage;
+	//~ End FSceneViewExtensionBase Interface
+};
+
diff --git a/Engine/Plugins/Media/MediaIOFramework/Source/MediaIOCore/Private/MediaCaptureSources.h b/Engine/Plugins/Media/MediaIOFramework/Source/MediaIOCore/Private/MediaCaptureSources.h
new file mode 100644
--- /dev/null
+++ b/Engine/Plugins/Media/MediaIOFramework/Source/MediaIOCore/Private/MediaCaptureSources.h
@@ -0,0 +1,421 @@
+// Copyright Epic Games, Inc. All Rights Reserved.
+
+#pragma once
+
+#include "CoreMinimal.h"
+#include "Engine/TextureRenderTarget2D.h"
+#include "Math/IntPoint.h"
+#include "MediaCapture.h"
+#include "MediaIOCoreModule.h"
+#include "MediaOutput.h"
+#include "Slate/SceneViewport.h"
+#include "TextureResource.h"
+#include "UObject/WeakObjectPtr.h"
+
+namespace UE::MediaCapture::Private
+{
+	bool ValidateSize(const FIntPoint TargetSize, const FIntPoint& DesiredSize, const FMediaCaptureOptions& CaptureOptions, const bool bCurrentlyCapturing)
+	{
+		if (CaptureOptions.Crop == EMediaCaptureCroppingType::None)
+		{
+			if (DesiredSize.X != TargetSize.X || DesiredSize.Y != TargetSize.Y)
+			{
+				UE_LOG(LogMediaIOCore, Error, TEXT("Can not %s the capture. The Render Target size doesn't match with the requested size. SceneViewport: %d,%d  MediaOutput: %d,%d")
+					, bCurrentlyCapturing ? TEXT("continue") : TEXT("start")
+					, TargetSize.X, TargetSize.Y
+					, DesiredSize.X, DesiredSize.Y);
+				return false;
+			}
+		}
+		else
+		{
+			FIntPoint StartCapturePoint = FIntPoint::ZeroValue;
+			if (CaptureOptions.Crop == EMediaCaptureCroppingType::Custom)
+			{
+				if (CaptureOptions.CustomCapturePoint.X < 0 || CaptureOptions.CustomCapturePoint.Y < 0)
+				{
+					UE_LOG(LogMediaIOCore, Error, TEXT("Can not %s the capture. The start capture point is negatif. Start Point: %d,%d")
+						, bCurrentlyCapturing ? TEXT("continue") : TEXT("start")
+						, StartCapturePoint.X, StartCapturePoint.Y);
+					return false;
+				}
+				StartCapturePoint = CaptureOptions.CustomCapturePoint;
+			}
+
+			if (DesiredSize.X + StartCapturePoint.X > TargetSize.X || DesiredSize.Y + StartCapturePoint.Y > TargetSize.Y)
+			{
+				UE_LOG(LogMediaIOCore, Error, TEXT("Can not %s the capture. The Render Target size is too small for the requested cropping options. SceneViewport: %d,%d  MediaOutput: %d,%d Start Point: %d,%d")
+					, bCurrentlyCapturing ? TEXT("continue") : TEXT("start")
+					, TargetSize.X, TargetSize.Y
+					, DesiredSize.X, DesiredSize.Y
+					, StartCapturePoint.X, StartCapturePoint.Y);
+				return false;
+			}
+		}
+
+		return true;
+	}
+
+	bool ValidateSceneViewport(const TSharedPtr<FSceneViewport>& SceneViewport, const FMediaCaptureOptions& CaptureOptions, const FIntPoint& DesiredSize, const EPixelFormat DesiredPixelFormat, const bool bCurrentlyCapturing)
+	{
+		if (!SceneViewport.IsValid())
+		{
+			UE_LOG(LogMediaIOCore, Error, TEXT("Can not %s the capture. The Scene Viewport is invalid.")
+				, bCurrentlyCapturing ? TEXT("continue") : TEXT("start"));
+			return false;
+		}
+
+		const FIntPoint SceneViewportSize = SceneViewport->GetRenderTargetTextureSizeXY();
+		if (!ValidateSize(SceneViewportSize, DesiredSize, CaptureOptions, bCurrentlyCapturing))
+		{
+			return false;
+		}
+
+		static const auto CVarDefaultBackBufferPixelFormat = IConsoleManager::Get().FindTConsoleVariableDataInt(TEXT("r.DefaultBackBufferPixelFormat"));
+		EPixelFormat SceneTargetFormat = EDefaultBackBufferPixelFormat::Convert2PixelFormat(EDefaultBackBufferPixelFormat::FromInt(CVarDefaultBackBufferPixelFormat->GetValueOnGameThread()));
+		if (DesiredPixelFormat != SceneTargetFormat)
+		{
+			if (!UMediaCapture::GetSupportedRgbaSwizzleFormats().Contains(SceneTargetFormat) || !CaptureOptions.bConvertToDesiredPixelFormat)
+			{
+				UE_LOG(LogMediaIOCore, Error, TEXT("Can not %s the capture. The Render Target pixel format doesn't match with the requested pixel format. %sRenderTarget: %s MediaOutput: %s")
+					, bCurrentlyCapturing ? TEXT("continue") : TEXT("start")
+					, (UMediaCapture::GetSupportedRgbaSwizzleFormats().Contains(SceneTargetFormat) && !CaptureOptions.bConvertToDesiredPixelFormat) ? TEXT("Please enable \"Convert To Desired Pixel Format\" option in Media Capture settings") : TEXT("")
+					, GetPixelFormatString(SceneTargetFormat)
+					, GetPixelFormatString(DesiredPixelFormat));
+				return false;
+			}
+			else
+			{
+				UE_LOG(LogMediaIOCore, Warning, TEXT("The Render Target pixel format doesn't match with the requested pixel format. Render target will be automatically converted. This could have a slight performance impact. RenderTarget: %s MediaOutput: %s")
+					, GetPixelFormatString(SceneTargetFormat)
+					, GetPixelFormatString(DesiredPixelFormat));
+			}
+		}
+
+		return true;
+	}
+	
+	bool ValidateTextureRenderTarget2D(const UTextureRenderTarget2D* InRenderTarget2D, const FMediaCaptureOptions& CaptureOptions, const FIntPoint& DesiredSize, const EPixelFormat DesiredPixelFormat, const bool bCurrentlyCapturing)
+	{
+		if (InRenderTarget2D == nullptr)
+		{
+			UE_LOG(LogMediaIOCore, Error, TEXT("Couldn't %s the capture. The Render Target is invalid.")
+				, bCurrentlyCapturing ? TEXT("continue") : TEXT("start"));
+			return false;
+		}
+
+		if (!ValidateSize(FIntPoint(InRenderTarget2D->SizeX, InRenderTarget2D->SizeY), DesiredSize, CaptureOptions, bCurrentlyCapturing))
+		{
+			return false;
+		}
+
+		if (DesiredPixelFormat != InRenderTarget2D->GetFormat())
+		{
+			if (!UMediaCapture::GetSupportedRgbaSwizzleFormats().Contains(InRenderTarget2D->GetFormat()) || !CaptureOptions.bConvertToDesiredPixelFormat)
+			{
+				UE_LOG(LogMediaIOCore, Error, TEXT("Can not %s the capture. The Render Target pixel format doesn't match with the requested pixel format. %sRenderTarget: %s MediaOutput: %s")
+					, bCurrentlyCapturing ? TEXT("continue") : TEXT("start")
+					, (UMediaCapture::GetSupportedRgbaSwizzleFormats().Contains(InRenderTarget2D->GetFormat()) && !CaptureOptions.bConvertToDesiredPixelFormat) ? TEXT("Please enable \"Convert To Desired Pixel Format\" option in Media Capture settings. ") : TEXT("")
+					, GetPixelFormatString(InRenderTarget2D->GetFormat())
+					, GetPixelFormatString(DesiredPixelFormat));
+				return false;
+			}
+			else
+			{
+				UE_LOG(LogMediaIOCore, Warning, TEXT("The Render Target pixel format doesn't match with the requested pixel format. Render target will be automatically converted. This could have a slight performance impact. RenderTarget: %s MediaOutput: %s")
+					, GetPixelFormatString(InRenderTarget2D->GetFormat())
+					, GetPixelFormatString(DesiredPixelFormat));
+			}
+		}
+
+		return true;
+	}
+
+	enum class ECaptureType : uint8
+	{
+		Immediate,
+		OnTick
+	};
+	
+	class FCaptureSource
+	{
+	public:
+		FCaptureSource(UMediaCapture* InMediaCapture, FMediaCaptureOptions InCaptureOptions)
+			: MediaCapture(InMediaCapture)
+			, CaptureOptions(MoveTemp(InCaptureOptions))
+		{
+		}
+
+		virtual ~FCaptureSource() = default;
+		
+		virtual void ResizeSourceBuffer(FIntPoint Size) {}
+		virtual void ResetSourceBufferSize(bool bFlushRenderingCommands) {}
+		virtual bool ValidateSource(const FIntPoint& DesiredSize, const EPixelFormat DesiredPixelFormat, const bool bCurrentlyCapturing) = 0;
+		virtual bool PostInitialize() = 0;
+		virtual EMediaCaptureSourceType GetSourceType() = 0;
+		virtual FIntPoint GetSize() = 0;
+		
+		virtual ECaptureType GetCaptureType() const
+		{
+			return ECaptureType::OnTick;
+		}
+		
+		virtual bool UpdateSourceImpl()
+		{
+			return true;
+		}
+		
+		/** Only returns a valid texture for scene viewports and render targets as the RHI Sources don't keep a pointer to the RHI texture. */
+		virtual FTexture2DRHIRef GetSourceTextureForInput_RenderThread(FRHICommandListImmediate& RHICmdList)
+		{
+			return nullptr;
+		}
+
+		virtual UTextureRenderTarget2D* GetRenderTarget() const
+		{
+			return nullptr;
+		}
+		
+		virtual TSharedPtr<FSceneViewport> GetSceneViewport() const
+		{
+			return nullptr;
+		}
+
+	public:
+		TWeakObjectPtr<UMediaCapture> MediaCapture;
+		FMediaCaptureOptions CaptureOptions;
+	};
+
+	class FSceneViewportCaptureSource : public FCaptureSource
+	{
+	public:
+		FSceneViewportCaptureSource(UMediaCapture* InMediaCapture, FMediaCaptureOptions InCaptureOptions, TSharedPtr<FSceneViewport> InSceneViewport)
+			: FCaptureSource(InMediaCapture, MoveTemp(InCaptureOptions))
+			, WeakViewport(InSceneViewport)
+		{
+		}
+
+		virtual ECaptureType GetCaptureType() const override
+		{
+			if (CaptureOptions.CapturePhase != EMediaCapturePhase::EndFrame)
+            {
+				return ECaptureType::Immediate;
+            }
+			return ECaptureType::OnTick;
+		}
+		
+		virtual EMediaCaptureSourceType GetSourceType() override
+		{
+			return EMediaCaptureSourceType::SCENE_VIEWPORT;
+		}
+		
+		virtual void ResetSourceBufferSize(bool bFlushRenderingCommands) override
+		{
+			if (TSharedPtr<FSceneViewport> SceneViewport = WeakViewport.Pin())
+			{
+				if (MediaCapture.IsValid())
+				{
+					MediaCapture->ResetFixedViewportSize(SceneViewport, bFlushRenderingCommands);
+				}
+			}				
+		}
+
+		virtual void ResizeSourceBuffer(FIntPoint Size) override
+		{
+			if (TSharedPtr<FSceneViewport> SceneViewport = WeakViewport.Pin())
+			{
+				if (MediaCapture.IsValid())
+				{
+					MediaCapture->SetFixedViewportSize(SceneViewport, Size);
+				}
+			}
+		}
+
+		virtual bool ValidateSource(const FIntPoint& InDesiredSize, const EPixelFormat InDesiredPixelFormat, const bool bInCurrentlyCapturing) override
+		{
+			bool bResult = false;
+			if (TSharedPtr<FSceneViewport> SceneViewport = WeakViewport.Pin())
+			{
+				bResult = ValidateSceneViewport(SceneViewport, CaptureOptions, InDesiredSize, InDesiredPixelFormat, bInCurrentlyCapturing);
+				if (!bResult)
+				{
+					MediaCapture->ResetFixedViewportSize(SceneViewport, false);
+				}
+			}
+
+			return bResult;
+		}
+
+		virtual bool PostInitialize() override
+		{
+			if (TSharedPtr<FSceneViewport> SceneViewport = WeakViewport.Pin())
+			{
+				return MediaCapture->PostInitializeCaptureViewport(SceneViewport);
+			}
+			return false;
+		}
+		
+		virtual FIntPoint GetSize() override
+		{
+			if (TSharedPtr<FSceneViewport> SceneViewport = WeakViewport.Pin())
+			{
+				return SceneViewport->GetSize();
+			}
+			return FIntPoint();
+		}
+
+		virtual bool UpdateSourceImpl() override
+		{
+			if (MediaCapture.IsValid())
+			{
+				if (TSharedPtr<FSceneViewport> SceneViewport = WeakViewport.Pin())
+				{
+					return MediaCapture->UpdateSceneViewportImpl(SceneViewport);
+				}
+			}
+
+			return false;
+		}
+
+		virtual FTexture2DRHIRef GetSourceTextureForInput_RenderThread(FRHICommandListImmediate& RHICmdList) override
+		{
+			FTexture2DRHIRef SourceTexture;
+			if (TSharedPtr<FSceneViewport> Viewport = WeakViewport.Pin())
+			{
+#if WITH_EDITOR
+				if (!IsRunningGame())
+				{
+					// PIE, PIE in windows, editor viewport
+					SourceTexture = Viewport->GetRenderTargetTexture();
+					if (!SourceTexture.IsValid() && Viewport->GetViewportRHI())
+					{
+						SourceTexture = RHICmdList.GetViewportBackBuffer(Viewport->GetViewportRHI());
+					}
+				}
+				else
+#endif
+				if (Viewport->GetViewportRHI())
+				{
+					// Standalone and packaged
+					SourceTexture = RHICmdList.GetViewportBackBuffer(Viewport->GetViewportRHI());
+				}
+			}
+
+			return SourceTexture;
+		}
+
+		virtual TSharedPtr<FSceneViewport> GetSceneViewport() const override
+		{
+			return WeakViewport.Pin();
+		}
+		
+		TWeakPtr<FSceneViewport> WeakViewport;
+	};
+
+	class FRHIResourceCaptureSource : public FCaptureSource
+	{
+	public:
+		FRHIResourceCaptureSource(UMediaCapture* InMediaCapture, FMediaCaptureOptions InCaptureOptions, const FRHICaptureResourceDescription& InResourceDescription)
+			: FCaptureSource(InMediaCapture, MoveTemp(InCaptureOptions))
+			, ResourceDescription(InResourceDescription)
+		{
+		}
+
+		virtual ECaptureType GetCaptureType() const override
+		{
+			return ECaptureType::Immediate;
+		}
+
+		virtual EMediaCaptureSourceType GetSourceType() override
+		{
+			return EMediaCaptureSourceType::RHI_RESOURCE;
+		}
+
+		virtual bool PostInitialize() override
+		{
+			return MediaCapture->PostInitializeCaptureRHIResource(ResourceDescription);
+		}
+
+		virtual FIntPoint GetSize() override
+		{
+			return ResourceDescription.ResourceSize;
+		}
+
+		virtual bool ValidateSource(const FIntPoint& DesiredSize, const EPixelFormat DesiredPixelFormat, const bool bCurrentlyCapturing) override
+		{
+			return true;
+		}
+
+		FRHICaptureResourceDescription ResourceDescription;
+	};
+
+
+	class FRenderTargetCaptureSource : public FCaptureSource
+	{
+	public:
+		FRenderTargetCaptureSource(UMediaCapture* InMediaCapture, FMediaCaptureOptions InCaptureOptions, UTextureRenderTarget2D* InRenderTarget)
+			: FCaptureSource(InMediaCapture, MoveTemp(InCaptureOptions))
+			, RenderTarget(InRenderTarget)
+		{
+		}
+		
+		virtual EMediaCaptureSourceType GetSourceType() override
+        {
+        	return EMediaCaptureSourceType::RENDER_TARGET;
+        }
+
+		virtual void ResizeSourceBuffer(FIntPoint Size) override
+		{
+			if (RenderTarget.IsValid())
+			{
+				RenderTarget->ResizeTarget(Size.X, Size.Y);
+			}
+		}
+		
+		virtual bool ValidateSource(const FIntPoint& InDesiredSize, const EPixelFormat InDesiredPixelFormat, const bool bInCurrentlyCapturing) override
+		{
+			return ValidateTextureRenderTarget2D(RenderTarget.Get(), CaptureOptions, InDesiredSize, InDesiredPixelFormat, bInCurrentlyCapturing);
+		}
+
+		virtual bool PostInitialize() override
+		{
+			return MediaCapture->PostInitializeCaptureRenderTarget(RenderTarget.Get());
+		}
+
+		virtual FIntPoint GetSize() override
+		{
+			if (RenderTarget.IsValid())
+			{
+				return FIntPoint(RenderTarget->SizeX, RenderTarget->SizeY);
+			}
+
+			return FIntPoint();
+		}
+		
+		virtual bool UpdateSourceImpl() override
+		{
+			if (MediaCapture.IsValid())
+			{
+				return MediaCapture->UpdateRenderTargetImpl(RenderTarget.Get());
+			}
+			return false;
+		}
+
+		virtual FTexture2DRHIRef GetSourceTextureForInput_RenderThread(FRHICommandListImmediate& RHICmdList)
+		{
+			if (RenderTarget.IsValid() && RenderTarget->GetRenderTargetResource())
+			{
+				return RenderTarget->GetRenderTargetResource()->GetTextureRenderTarget2DResource()->GetTextureRHI();
+			}
+			return nullptr;
+		}
+
+		virtual UTextureRenderTarget2D* GetRenderTarget() const override
+		{
+			return RenderTarget.Get();
+		}
+
+		TWeakObjectPtr<UTextureRenderTarget2D> RenderTarget;
+	};
+}
+
diff --git a/Engine/Plugins/Media/MediaIOFramework/Source/MediaIOCore/Private/MediaIOFrameManager.h b/Engine/Plugins/Media/MediaIOFramework/Source/MediaIOCore/Private/MediaIOFrameManager.h
new file mode 100644
--- /dev/null
+++ b/Engine/Plugins/Media/MediaIOFramework/Source/MediaIOCore/Private/MediaIOFrameManager.h
@@ -0,0 +1,129 @@
+// Copyright Epic Games, Inc. All Rights Reserved.
+
+#pragma once
+
+#include "CoreMinimal.h"
+#include "RenderingThread.h"
+
+namespace UE::MediaCaptureData
+{
+	class FFrame
+	{
+	public:
+		virtual ~FFrame() {}
+		virtual int32 GetId() const = 0;
+		virtual bool IsPending() const = 0;
+	};
+
+	/**
+	 * Manages a queue of frames used by MediaCapture.
+	 */
+	class FFrameManager
+	{
+	public:
+		~FFrameManager()
+		{
+			ENQUEUE_RENDER_COMMAND(MediaCaptureFrameManagerCleaning)(
+				[FramesToBeReleased = MoveTemp(Frames)](FRHICommandListImmediate& RHICmdList) mutable
+			{
+				FramesToBeReleased.Reset();
+			});
+		}
+
+		void ForEachFrame(TFunctionRef<void(const TSharedPtr<FFrame>&)> Function)
+		{
+			for (const TSharedPtr<FFrame>& Frame : Frames)
+			{
+				Function(Frame);
+			}
+		}
+
+		FString GetFramesState()
+		{
+			TStringBuilder<256> ReadbackInfoBuilder;
+			ReadbackInfoBuilder << "\n";
+			for (int32 Index = 0; Index < Frames.Num(); Index++)
+			{
+				ReadbackInfoBuilder << FString::Format(TEXT("Frame {0} State:  {1}\n"), { Index, Frames[Index]->IsPending() ? TEXT("Pending") : TEXT("Idle") });
+			}
+
+			return ReadbackInfoBuilder.ToString();
+		}
+
+		TConstArrayView<TSharedPtr<FFrame>> GetFrames() const
+		{
+			return Frames;
+		}
+
+		void AddFrame(TSharedPtr<FFrame> NewFrame)
+		{
+			AvailableFrames.Enqueue(NewFrame->GetId());
+			Frames.Add(MoveTemp(NewFrame));
+		}
+
+		void MarkAvailable(const FFrame& InFrame)
+		{
+			if (ensure(!InFrame.IsPending()))
+			{
+				AvailableFrames.Enqueue(InFrame.GetId());
+			}
+		}
+
+		void MarkPending(const FFrame& InFrame)
+		{
+			if (ensure(InFrame.IsPending()))
+			{
+				PendingFrames.Enqueue(InFrame.GetId());
+			}
+		}
+
+		void CompleteNextPending(const FFrame& InFrame)
+		{
+			TOptional<int32> NextPending = PendingFrames.Dequeue();
+			if (ensure(NextPending.IsSet() && InFrame.GetId() == NextPending.GetValue()))
+			{
+				MarkAvailable(InFrame);
+			}
+		}
+
+		template<typename FrameType>
+		TSharedPtr<FrameType> GetNextAvailable()
+		{
+			TSharedPtr<FrameType> NextAvailableFrame;
+			
+			static_assert(TIsDerivedFrom<FrameType, FFrame>::Value, "The frame type must derive from FFrame.");
+			TOptional<int32> NextAvailable = AvailableFrames.Dequeue();
+			if (NextAvailable.IsSet())
+			{
+				NextAvailableFrame = StaticCastSharedPtr<FrameType>(Frames[NextAvailable.GetValue()]);
+			}
+
+			// Peek next pending before capturing to avoid getting the newly captured one. 
+			// It shouldn't happen but if it happens it will start waiting for the capture to complete
+			TSharedPtr<FFrame> NextPendingFrame = PeekNextPending<FFrame>();
+			if (NextPendingFrame && NextAvailableFrame)
+			{
+				ensureMsgf(NextPendingFrame->GetId() != NextAvailableFrame->GetId(), TEXT("The next pending frame was the same as the next available frame, this should not have happened!"));
+			}
+
+			return NextAvailableFrame;
+		}
+
+		template<typename FrameType>
+		TSharedPtr<FrameType> PeekNextPending()
+		{
+			static_assert(TIsDerivedFrom<FrameType, FFrame>::Value, "The frame type must derive from FFrame.");
+			if (const int32* NextPending = PendingFrames.Peek())
+			{
+				return StaticCastSharedPtr<FrameType>(Frames[*NextPending]);
+			}
+
+			return nullptr;
+		}
+
+	private:
+		TArray<TSharedPtr<FFrame>> Frames;
+		TSpscQueue<int32> AvailableFrames;
+		TSpscQueue<int32> PendingFrames;
+	};
+}
diff --git a/Engine/Plugins/Media/MediaIOFramework/Source/MediaIOCore/Public/MediaCapture.h b/Engine/Plugins/Media/MediaIOFramework/Source/MediaIOCore/Public/MediaCapture.h
--- a/Engine/Plugins/Media/MediaIOFramework/Source/MediaIOCore/Public/MediaCapture.h
+++ b/Engine/Plugins/Media/MediaIOFramework/Source/MediaIOCore/Public/MediaCapture.h
@@ -3,7 +3,6 @@
 #pragma once
 
 #include "CoreTypes.h"
-#include "UObject/ObjectMacros.h"
 #include "UObject/Object.h"
 
 #include <atomic>
@@ -12,12 +11,13 @@
 #include "MediaOutput.h"
 #include "Misc/Timecode.h"
 #include "PixelFormat.h"
-#include "RendererInterface.h"
 #include "RenderGraphResources.h"
 #include "RHI.h"
 #include "RHIResources.h"
+#include "Tasks/Task.h"
 #include "Templates/PimplPtr.h"
 
+
 #include "MediaCapture.generated.h"
 
 class FSceneViewport;
@@ -27,12 +27,19 @@ class FRHICommandListImmediate;
 class FRHIGPUTextureReadback;
 class FRHIGPUBufferReadback;
 class FRDGBuilder;
+class UMediaCapture;
 
 namespace UE::MediaCaptureData
 {
 	class FFrameManager;
 	class FCaptureFrame;
 	class FMediaCaptureHelper;
+	struct FCaptureFrameArgs;
+}
+
+namespace UE::MediaCapture::Private
+{
+	class FCaptureSource;
 }
 
 /**
@@ -75,7 +82,6 @@ enum class EMediaCaptureResourceType : uint8
  */
 class FMediaCaptureUserData
 {
-
 };
 
 /**
@@ -114,6 +120,16 @@ struct MEDIAIOCORE_API FRHICaptureResourceDescription
 	FIntPoint ResourceSize = FIntPoint::ZeroValue;
 };
 
+UENUM(BlueprintType)
+enum class EMediaCapturePhase : uint8
+{
+	BeforePostProcessing,
+	AfterToneMap,
+	AfterFXAA,
+	AfterMotionBlur,
+	EndFrame
+};
+
 /**
  * Base class of additional data that can be stored for each requested capture.
  */
@@ -185,6 +201,9 @@ public:
 	/** The number of images to capture*/
 	UPROPERTY(BlueprintReadWrite, EditAnywhere, Category = "MediaCapture", meta=(editcondition="bAutostopOnCapture"))
 	int32 NumberOfFramesToCapture;
+
+	UPROPERTY(BlueprintReadWrite, EditAnywhere, Category = "MediaCapture")
+	EMediaCapturePhase CapturePhase = EMediaCapturePhase::EndFrame;
 };
 
 
@@ -192,7 +211,6 @@ public:
 DECLARE_DYNAMIC_MULTICAST_DELEGATE(FMediaCaptureStateChangedSignature);
 DECLARE_MULTICAST_DELEGATE(FMediaCaptureStateChangedSignatureNative);
 
-
 /**
  * Abstract base class for media capture.
  *
@@ -202,9 +220,10 @@ DECLARE_MULTICAST_DELEGATE(FMediaCaptureStateChangedSignatureNative);
 UCLASS(Abstract, editinlinenew, BlueprintType, hidecategories = (Object))
 class MEDIAIOCORE_API UMediaCapture : public UObject
 {
-	GENERATED_UCLASS_BODY()
+	GENERATED_BODY()
 
 public:
+	UMediaCapture();
 
 	/** Default constructor / destructor to use forward declared uniqueptr */
 	virtual ~UMediaCapture();
@@ -266,6 +285,8 @@ public:
 	
 	/** Captures a resource immediately from the render thread. Used in RHI_RESOURCE capture mode */
 	void CaptureImmediate_RenderThread(FRDGBuilder& GraphBuilder, FRHITexture* InSourceTexture);
+	
+	void CaptureImmediate_RenderThread(FRDGBuilder& GraphBuilder, FRDGTextureRef InSourceTextureRef);
 
 	/**
 	 * Stop the previous requested capture.
@@ -306,6 +327,23 @@ public:
 	 */
 	FMediaCaptureStateChangedSignatureNative OnStateChangedNative;
 
+	void SetFixedViewportSize(TSharedPtr<FSceneViewport> InSceneViewport, FIntPoint InSize);
+	void ResetFixedViewportSize(TSharedPtr<FSceneViewport> InViewport, bool bInFlushRenderingCommands);
+
+	/** Called after initialize for viewport capture type */
+	virtual bool PostInitializeCaptureViewport(TSharedPtr<FSceneViewport>& InSceneViewport) { return true; }
+
+	/** Called after initialize for render target capture type */
+	virtual bool PostInitializeCaptureRenderTarget(UTextureRenderTarget2D* InRenderTarget) { return true; }
+
+	/** Called after initialize for rhi resource capture type */
+	virtual bool PostInitializeCaptureRHIResource(const FRHICaptureResourceDescription& InResourceDescription) { return true; }
+	
+	virtual bool UpdateSceneViewportImpl(TSharedPtr<FSceneViewport>& InSceneViewport) { return true; }
+	virtual bool UpdateRenderTargetImpl(UTextureRenderTarget2D* InRenderTarget) { return true; }
+
+	static const TSet<EPixelFormat>& GetSupportedRgbaSwizzleFormats();
+
 public:
 	//~ UObject interface
 	virtual void BeginDestroy() override;
@@ -324,17 +362,6 @@ protected:
 	/** Initialization method to prepare implementation for capture */
 	virtual bool InitializeCapture() PURE_VIRTUAL(UMediaCapture::InitializeCapture, return false; );
 	
-	/** Called after initialize for viewport capture type */
-	virtual bool PostInitializeCaptureViewport(TSharedPtr<FSceneViewport>& InSceneViewport) { return true; }
-
-	/** Called after initialize for render target capture type */
-	virtual bool PostInitializeCaptureRenderTarget(UTextureRenderTarget2D* InRenderTarget) { return true; }
-
-	/** Called after initialize for rhi resource capture type */
-	virtual bool PostInitializeCaptureRHIResource(const FRHICaptureResourceDescription& InResourceDescription) { return true; }
-
-	virtual bool UpdateSceneViewportImpl(TSharedPtr<FSceneViewport>& InSceneViewport) { return true; }
-	virtual bool UpdateRenderTargetImpl(UTextureRenderTarget2D* InRenderTarget) { return true; }
 	virtual void StopCaptureImpl(bool bAllowPendingFrameToBeProcess) { }
 
 	//~ DMA Functions implemented in child classes, Temporary until AJA and BM middleman are migrated to a module.
@@ -345,12 +372,10 @@ protected:
 	friend class UE::MediaCaptureData::FMediaCaptureHelper;
 	struct FCaptureBaseData
 	{
-		FCaptureBaseData();
-
 		FTimecode SourceFrameTimecode;
 		FFrameRate SourceFrameTimecodeFramerate;
-		uint32 SourceFrameNumberRenderThread;
-		uint32 SourceFrameNumber;
+		uint32 SourceFrameNumberRenderThread = 0;
+		uint32 SourceFrameNumber = 0;
 	};
 
 	/**
@@ -378,6 +403,23 @@ protected:
 	 */
 	virtual void OnFrameCaptured_RenderingThread(const FCaptureBaseData& InBaseData, TSharedPtr<FMediaCaptureUserData, ESPMode::ThreadSafe> InUserData, void* InBuffer, int32 Width, int32 Height, int32 BytesPerRow) { }
 
+	/**
+	 * Args passed to OnFrameCaptured
+	 **/
+	struct FMediaCaptureResourceData
+	{
+		void* Buffer = nullptr;
+		int32 Width = 0;
+		int32 Height = 0;
+		int32 BytesPerRow = 0;
+	};
+
+	/**
+	 * Callback when the buffer was successfully copied to CPU ram.
+	 * The buffer is only valid for the duration of the callback.
+	 */
+	virtual void OnFrameCaptured_AnyThread(const FCaptureBaseData& InBaseData, TSharedPtr<FMediaCaptureUserData, ESPMode::ThreadSafe> InUserData, const FMediaCaptureResourceData& InResourceData) { }
+
 	UE_DEPRECATED(5.1, "OnRHITextureCaptured_RenderingThread has been deprecated to support texture and buffer output resource. Please use OnRHIResourceCaptured_RenderingThread instead.")
 	virtual void OnRHITextureCaptured_RenderingThread(const FCaptureBaseData& InBaseData, TSharedPtr<FMediaCaptureUserData, ESPMode::ThreadSafe> InUserData, FTextureRHIRef InTexture) { }
 
@@ -388,6 +430,10 @@ protected:
 	 */
 	virtual void OnRHIResourceCaptured_RenderingThread(const FCaptureBaseData& InBaseData, TSharedPtr<FMediaCaptureUserData, ESPMode::ThreadSafe> InUserData, FTextureRHIRef InTexture) { }
 	virtual void OnRHIResourceCaptured_RenderingThread(const FCaptureBaseData& InBaseData, TSharedPtr<FMediaCaptureUserData, ESPMode::ThreadSafe> InUserData, FBufferRHIRef InBuffer) { }
+	
+	virtual void OnRHIResourceCaptured_AnyThread(const FCaptureBaseData& InBaseData, TSharedPtr<FMediaCaptureUserData, ESPMode::ThreadSafe> InUserData, FTextureRHIRef InTexture) { }
+	virtual void OnRHIResourceCaptured_AnyThread(const FCaptureBaseData& InBaseData, TSharedPtr<FMediaCaptureUserData, ESPMode::ThreadSafe> InUserData, FBufferRHIRef InBuffer) { }
+	
 
 	UE_DEPRECATED(5.1, "OnCustomCapture_RenderingThread has been deprecated while moving to a RDG based pipeline. Please use the RDG version instead.")
 	virtual void OnCustomCapture_RenderingThread(FRHICommandListImmediate& RHICmdList, const FCaptureBaseData& InBaseData, TSharedPtr<FMediaCaptureUserData, ESPMode::ThreadSafe> InUserData, FTextureRHIRef InSourceTexture, FTextureRHIRef TargetableTexture, FResolveParams& ResolveParams, FVector2D CropU, FVector2D CropV)
@@ -443,10 +489,10 @@ protected:
 
 
 protected:
-	UTextureRenderTarget2D* GetTextureRenderTarget() { return CapturingRenderTarget; }
-	TSharedPtr<FSceneViewport> GetCapturingSceneViewport() { return CapturingSceneViewport.Pin(); }
+	UTextureRenderTarget2D* GetTextureRenderTarget() const;
+	TSharedPtr<FSceneViewport> GetCapturingSceneViewport() const;
 	EMediaCaptureConversionOperation GetConversionOperation() const { return ConversionOperation; }
-	const FString& GetCaptureSourceType() const;
+	FString GetCaptureSourceType() const;
 	void SetState(EMediaCaptureState InNewState);
 
 
@@ -463,14 +509,33 @@ private:
 	EMediaCaptureResourceType GetOutputResourceType(const EMediaCaptureConversionOperation & InConversionOperation) const;
 	FRDGBufferDesc GetOutputBufferDescription(EMediaCaptureConversionOperation InConversionOperation) const;
 	void BroadcastStateChanged();
-	void SetFixedViewportSize(TSharedPtr<FSceneViewport> InSceneViewport);
-	void ResetFixedViewportSize(TSharedPtr<FSceneViewport> InViewport, bool bInFlushRenderingCommands);
+	void WaitForPendingTasks();
+	
+	void WaitForSingleExperimentalSchedulingTaskToComplete();
+	void WaitForAllExperimentalSchedulingTasksToComplete();
+	void CleanupCompletedExperimentalSchedulingTasks();
+	
+	void CaptureImmediate_RenderThread(const UE::MediaCaptureData::FCaptureFrameArgs& Args);
 
-	bool ProcessCapture_RenderThread(FRDGBuilder& GraphBuilder, UMediaCapture* InMediaCapture, UE::MediaCaptureData::FCaptureFrame* CapturingFrame, FTexture2DRHIRef InResourceToCapture, FIntPoint InDesiredSize);
-	bool ProcessReadyFrame_RenderThread(FRHICommandListImmediate& RHICmdList, UMediaCapture* InMediaCapture, UE::MediaCaptureData::FCaptureFrame* ReadyFrame);
+
+	// Capture pipeline stuff
+	void ProcessCapture_GameThread();
+	void PrepareAndDispatchCapture_GameThread(const TSharedPtr<UE::MediaCaptureData::FCaptureFrame>& CapturingFrame);
+	
+	bool ProcessCapture_RenderThread(const TSharedPtr<UE::MediaCaptureData::FCaptureFrame>& CapturingFrame, const UE::MediaCaptureData::FCaptureFrameArgs& Args);
+	bool ProcessReadyFrame_RenderThread(FRHICommandListImmediate& RHICmdList, UMediaCapture* InMediaCapture, const TSharedPtr<UE::MediaCaptureData::FCaptureFrame>& ReadyFrame);
+	
+	void InitializeSyncHandlers_RenderThread();
+	void InitializeCaptureFrame(const TSharedPtr<UE::MediaCaptureData::FCaptureFrame>& CaptureFrame);
+
+	struct FMediaCaptureSyncData;
+	TSharedPtr<FMediaCaptureSyncData> GetAvailableSyncHandler() const;
 
 	void PrintFrameState();
 
+	bool StartSourceCapture(TSharedPtr<UE::MediaCapture::Private::FCaptureSource> InSource);
+	bool UpdateSource(TSharedPtr<UE::MediaCapture::Private::FCaptureSource> InCaptureSource);
+
 protected:
 
 	/** MediaOutput associated with this capture */
@@ -486,17 +551,12 @@ protected:
 	/** Type of capture currently being done. i.e Viewport, RT, RHI Resource */
 	EMediaCaptureSourceType CaptureSourceType = EMediaCaptureSourceType::SCENE_VIEWPORT;
 
-private:
-
-	UPROPERTY(Transient)
-	TObjectPtr<UTextureRenderTarget2D> CapturingRenderTarget;
-	
+private:	
 	TPimplPtr<UE::MediaCaptureData::FFrameManager> FrameManager;
 	int32 NumberOfCaptureFrame = 2;
 	int32 CaptureRequestCount = 0;
 	EMediaCaptureState MediaState = EMediaCaptureState::Stopped;
 
-	TWeakPtr<FSceneViewport> CapturingSceneViewport;
 	FCriticalSection AccessingCapturingSource;
 
 	FIntPoint DesiredSize = FIntPoint(1920, 1080);
@@ -511,8 +571,10 @@ private:
 
 	bool bViewportHasFixedViewportSize = false;
 	std::atomic<bool> bOutputResourcesInitialized;
+	std::atomic<bool> bSyncHandlersInitialized;
 	std::atomic<bool> bShouldCaptureRHIResource;
 	std::atomic<int32> WaitingForRenderCommandExecutionCounter;
+	TArray<TFuture<void>> PendingCommands;
 	std::atomic<int32> PendingFrameCount;
 
 	struct FQueuedCaptureData
@@ -521,4 +583,35 @@ private:
 		TSharedPtr<FMediaCaptureUserData> UserData;
 	};
 	TSpscQueue<FQueuedCaptureData> CaptureDataQueue;
+
+	/** Structure holding data used for synchronization of buffer output */
+	struct FMediaCaptureSyncData
+	{
+		/**
+		 * We use a RHIFence to write to in a pass following the buffer conversion one.
+		 * We spawn a task that will poll for the fence to be over (Waiting is not exposed)
+		 * and then we will push the captured buffer to the capture callback
+		 */
+		FGPUFenceRHIRef RHIFence;
+
+		/** Whether fence is currently waiting to be cleared */
+		std::atomic<bool> bIsBusy = false;
+	};
+
+	/** Array of sync handlers (fence) to sync when captured buffer is completed */
+	TArray<TSharedPtr<FMediaCaptureSyncData>> SyncHandlers;
+
+	/** Whether capture is active. Used to be queried by sync task */
+	std::atomic<bool> bIsActive;
+
+	FCriticalSection PendingReadbackTasksCriticalSection;
+
+	/** List of pending readback tasks, when CVarMediaIOEnableExperimentalScheduling and CVarMediaIOScheduleOnAnyThread are set to true. */
+	TArray<UE::Tasks::FTask> PendingReadbackTasks;
+
+	TSharedPtr<UE::MediaCapture::Private::FCaptureSource> CaptureSource;
+
+	/** Holds a view extension that callbacks  */
+	TSharedPtr<class FMediaCaptureSceneViewExtension> ViewExtension; 
 };
+ 
\ No newline at end of file
diff --git a/Engine/Plugins/Media/MediaIOFramework/Source/MediaIOCore/Public/MediaOutput.h b/Engine/Plugins/Media/MediaIOFramework/Source/MediaIOCore/Public/MediaOutput.h
--- a/Engine/Plugins/Media/MediaIOFramework/Source/MediaIOCore/Public/MediaOutput.h
+++ b/Engine/Plugins/Media/MediaIOFramework/Source/MediaIOCore/Public/MediaOutput.h
@@ -23,6 +23,7 @@ enum class EMediaCaptureConversionOperation : int8
 	CUSTOM, //~ If Custom used, OnCustomCapture_RenderingThread, GetCustomOutputSize and GetCustomOutputPixelFormat on the MediaCapture need to be implemented
 };
 
+UENUM()
 enum class EMediaCaptureSourceType : int8
 {
 	RENDER_TARGET,
