GPUTextureTransfer - Return nullptr when calling GetTextureTransfer when running without an RHI

diff --git a/Engine/Plugins/Media/MediaIOFramework/Source/GPUTextureTransfer/Private/GPUTextureTransferModule.cpp b/Engine/Plugins/Media/MediaIOFramework/Source/GPUTextureTransfer/Private/GPUTextureTransferModule.cpp
--- a/Engine/Plugins/Media/MediaIOFramework/Source/GPUTextureTransfer/Private/GPUTextureTransferModule.cpp
+++ b/Engine/Plugins/Media/MediaIOFramework/Source/GPUTextureTransfer/Private/GPUTextureTransferModule.cpp
@@ -80,20 +80,23 @@ void FGPUTextureTransferModule::ShutdownModule()
 
 UE::GPUTextureTransfer::TextureTransferPtr FGPUTextureTransferModule::GetTextureTransfer()
 {
-#if DVP_SUPPORTED_PLATFORM
-	UE::GPUTextureTransfer::ERHI SupportedRHI = ConvertRHI(RHIGetInterfaceType());
-	if (SupportedRHI == UE::GPUTextureTransfer::ERHI::Invalid) 
+	if (FApp::CanEverRender())
 	{
-		UE_LOG(LogGPUTextureTransfer, Error, TEXT("The current RHI is not supported with GPU Texture Transfer."));
-		return nullptr;
-	}
+#if DVP_SUPPORTED_PLATFORM
+		UE::GPUTextureTransfer::ERHI SupportedRHI = ConvertRHI(RHIGetInterfaceType());
+		if (SupportedRHI == UE::GPUTextureTransfer::ERHI::Invalid) 
+		{
+			UE_LOG(LogGPUTextureTransfer, Error, TEXT("The current RHI is not supported with GPU Texture Transfer."));
+			return nullptr;
+		}
 	
-	const uint8 RHIIndex = static_cast<uint8>(SupportedRHI);
-	if (TransferObjects[RHIIndex])
-	{
-		return TransferObjects[RHIIndex];
-	}
+		const uint8 RHIIndex = static_cast<uint8>(SupportedRHI);
+		if (TransferObjects[RHIIndex])
+		{
+			return TransferObjects[RHIIndex];
+		}
 #endif
+	}
 	return nullptr;
 }
 
