Tentative fix for crash while doing a media capture

diff --git a/Engine/Plugins/Media/MediaIOFramework/Source/MediaIOCore/Private/MediaCapture.cpp b/Engine/Plugins/Media/MediaIOFramework/Source/MediaIOCore/Private/MediaCapture.cpp
--- a/Engine/Plugins/Media/MediaIOFramework/Source/MediaIOCore/Private/MediaCapture.cpp
+++ b/Engine/Plugins/Media/MediaIOFramework/Source/MediaIOCore/Private/MediaCapture.cpp
@@ -576,7 +576,8 @@ namespace UE::MediaCaptureData
 		static void OnReadbackComplete(FRHICommandList& RHICmdList, UMediaCapture* MediaCapture, TSharedPtr<CaptureType> ReadyFrame)
 		{
 			UE_LOG(LogMediaIOCore, Verbose, TEXT("[%s - %s] - Processing pending frame %d"), *MediaCapture->MediaOutputName, *FThreadManager::GetThreadName(FPlatformTLS::GetCurrentThreadId()), ReadyFrame->GetId());
-			
+			TRACE_CPUPROFILER_EVENT_SCOPE(UMediaCapture::OnReadbackComplete);
+
 			{
 				ON_SCOPE_EXIT
 				{
@@ -653,7 +654,7 @@ namespace UE::MediaCaptureData
 					if (ColorDataBuffer)
 					{
 						{
-							TRACE_CPUPROFILER_EVENT_SCOPE_TEXT(*FString::Printf(TEXT("MediaCapture Output Frame %d"), ReadyFrame->CaptureBaseData.SourceFrameNumberRenderThread));
+							TRACE_CPUPROFILER_EVENT_SCOPE_TEXT(*FString::Printf(TEXT("MediaCapture Output Frame %d"), ReadyFrame->GetId()));
 							SCOPE_CYCLE_COUNTER(STAT_MediaCapture_AnyThread_LockResource)
 
 							if (IsInRenderingThread())
@@ -725,13 +726,40 @@ namespace UE::MediaCaptureData
 									TRACE_CPUPROFILER_EVENT_SCOPE(UMediaCapture::SyncPass);
 
 									double WaitTime = 0.0;
-
+									bool bWaitedForCompletion = false;
 									{
 										FScopedDurationTimer Timer(WaitTime);
 									
 										// Wait until fence has been written (shader has completed)
-										while (SyncDataPtr->RHIFence->Poll() == false && WeakCapture->bIsActive)
+										while (true)
 										{
+											if (SyncDataPtr->RHIFence->Poll())
+											{
+												bWaitedForCompletion = true;
+												break;
+											}
+
+											// Can't check for object validity during garbage collection.
+											if (!IsGarbageCollecting())
+											{
+
+												constexpr bool bEvenIfPendingKill = false;
+												constexpr bool ThreadSafeTest = true;
+												if (WeakCapture.IsValid(bEvenIfPendingKill, ThreadSafeTest))
+												{
+													if (!WeakCapture->bIsActive)
+													{
+														bWaitedForCompletion = false;
+														break;
+													}
+												}
+												else
+												{
+													bWaitedForCompletion = false;
+													break;
+												}
+											}
+
 											FPlatformProcess::SleepNoStats(0);
 										}
 
@@ -739,9 +767,7 @@ namespace UE::MediaCaptureData
 										SyncDataPtr->bIsBusy = false;
 									}
 
-									UE_LOG(LogMediaIOCore, VeryVerbose, TEXT("AsyncTask wait time: %f s"), WaitTime);
-
-									if (WeakCapture->bIsActive)
+									if (WeakCapture->bIsActive && bWaitedForCompletion)
 									{
 										if (CVarMediaIOScheduleOnAnyThread.GetValueOnAnyThread() == 1)
 										{
@@ -761,7 +787,7 @@ namespace UE::MediaCaptureData
 										}
 									}
 
-								}, LowLevelTasks::ETaskPriority::BackgroundHigh);
+								});
 
 							{
 								// Add the task to the capture's list of pending tasks.
@@ -1722,12 +1748,14 @@ void UMediaCapture::PrepareAndDispatchCapture_GameThread(const TSharedPtr<UE::Me
 					GraphBuilder.Execute();
 				}
 				
-				//Process the next pending frame
-				
-				UE_LOG(LogMediaIOCore, Verbose, TEXT("[%s] - Processing pending frame %d"), *InMediaCapture->MediaOutputName, NextPending ? NextPending->FrameId : -1);
-				if (NextPending)
+				if (CVarMediaIOEnableExperimentalScheduling.GetValueOnAnyThread() == 0)
 				{
-					InMediaCapture->ProcessReadyFrame_RenderThread(RHICmdList, InMediaCapture, NextPending);
+					//Process the next pending frame
+					UE_LOG(LogMediaIOCore, Verbose, TEXT("[%s] - Processing pending frame %d"), *InMediaCapture->MediaOutputName, NextPending ? NextPending->FrameId : -1);
+					if (NextPending)
+					{
+						InMediaCapture->ProcessReadyFrame_RenderThread(RHICmdList, InMediaCapture, NextPending);
+					}
 				}
 
 				// Whatever happens, we want to decrement our counter to track enqueued commands
diff --git a/Engine/Plugins/Media/MediaIOFramework/Source/MediaIOCore/Public/MediaCapture.h b/Engine/Plugins/Media/MediaIOFramework/Source/MediaIOCore/Public/MediaCapture.h
--- a/Engine/Plugins/Media/MediaIOFramework/Source/MediaIOCore/Public/MediaCapture.h
+++ b/Engine/Plugins/Media/MediaIOFramework/Source/MediaIOCore/Public/MediaCapture.h
@@ -602,7 +602,7 @@ private:
 	TArray<TSharedPtr<FMediaCaptureSyncData>> SyncHandlers;
 
 	/** Whether capture is active. Used to be queried by sync task */
-	std::atomic<bool> bIsActive;
+	std::atomic<bool> bIsActive = false;
 
 	FCriticalSection PendingReadbackTasksCriticalSection;
 
