MediaCapture - Don't access umediacapture in busy loop
Enable ScheduleOnAnyThread by default (MediaIO.ScheduleOnAnyThread)

diff --git a/Engine/Plugins/Media/MediaIOFramework/Source/MediaIOCore/Private/MediaCapture.cpp b/Engine/Plugins/Media/MediaIOFramework/Source/MediaIOCore/Private/MediaCapture.cpp
--- a/Engine/Plugins/Media/MediaIOFramework/Source/MediaIOCore/Private/MediaCapture.cpp
+++ b/Engine/Plugins/Media/MediaIOFramework/Source/MediaIOCore/Private/MediaCapture.cpp
@@ -48,7 +48,7 @@ static TAutoConsoleVariable<int32> CVarMediaIOEnableExperimentalScheduling(
 	ECVF_RenderThreadSafe);
 
 static TAutoConsoleVariable<int32> CVarMediaIOScheduleOnAnyThread(
-	TEXT("MediaIO.ScheduleOnAnyThread"), 0,
+	TEXT("MediaIO.ScheduleOnAnyThread"), 1,
 	TEXT("Whether to wait for resource readback in a separate thread. (Experimental)"),
 	ECVF_RenderThreadSafe);
 
@@ -111,8 +111,6 @@ namespace UE::MediaCaptureData
 	public:
 		FCaptureFrame(int32 InFrameId)
 			: FrameId(InFrameId)
-			, bReadbackRequested(false)
-			, bDoingGPUCopy(false)
 		{
 		}
 
@@ -156,8 +154,9 @@ namespace UE::MediaCaptureData
 
 		int32 FrameId = 0;
 		UMediaCapture::FCaptureBaseData CaptureBaseData;
-		TAtomic<bool> bReadbackRequested;
-		TAtomic<bool> bDoingGPUCopy;
+		std::atomic<bool> bReadbackRequested = false;
+		std::atomic<bool> bDoingGPUCopy = false;
+		std::atomic<bool> bMediaCaptureActive = true;
 		TSharedPtr<FMediaCaptureUserData> UserData;
 	};
 
@@ -686,7 +685,7 @@ namespace UE::MediaCaptureData
 
 
 		template <typename CaptureType>
-		static void AddSyncPointPass(FRDGBuilder& GraphBuilder, const TObjectPtr<UMediaCapture>& MediaCapture, TSharedPtr<FCaptureFrame> CapturingFrame, FRDGViewableResource* OutputResource)
+		static void AddSyncPointPass(FRDGBuilder& GraphBuilder, UMediaCapture* MediaCapture, TSharedPtr<FCaptureFrame> CapturingFrame, FRDGViewableResource* OutputResource)
 		{
 			TRACE_CPUPROFILER_EVENT_SCOPE(MediaCaptureSyncPoint);
 			TRACE_CPUPROFILER_EVENT_SCOPE_TEXT(*FString::Printf(TEXT("MediaCaptureSyncPoint_%d"), CapturingFrame->CaptureBaseData.SourceFrameNumberRenderThread));
@@ -702,18 +701,16 @@ namespace UE::MediaCaptureData
 			typename CaptureType::PassParametersType* PassParameters = GraphBuilder.AllocParameters<typename CaptureType::PassParametersType>();
 			PassParameters->Resource = static_cast<typename CaptureType::FOutputResourceType>(OutputResource);
 
-			TWeakObjectPtr<UMediaCapture> WeakCapture = MediaCapture;
-
 			GraphBuilder.AddPass(
 				RDG_EVENT_NAME("MediaCaptureCopySyncPass"),
 				PassParameters,
 				ERDGPassFlags::Copy | ERDGPassFlags::NeverCull,
-				[WeakCapture, CapturingFrame](FRHICommandListImmediate& RHICmdList)
+				[MediaCapture, CapturingFrame](FRHICommandListImmediate& RHICmdList)
 				{
-					if (WeakCapture.IsValid() && WeakCapture->bIsActive)
+					if (CapturingFrame && CapturingFrame->bMediaCaptureActive)
 					{
 						// Get available sync handler to create a sync point
-						TSharedPtr<UMediaCapture::FMediaCaptureSyncData> SyncDataPtr = WeakCapture->GetAvailableSyncHandler();
+						TSharedPtr<UMediaCapture::FMediaCaptureSyncData> SyncDataPtr = MediaCapture->GetAvailableSyncHandler();
 						if (ensure(SyncDataPtr))
 						{
 							// This will happen after the conversion pass has completed
@@ -721,7 +718,7 @@ namespace UE::MediaCaptureData
 							SyncDataPtr->bIsBusy = true;
 
 							// Spawn a task that will wait (poll) and continue the process of providing a new texture
-							UE::Tasks::FTask Task = UE::Tasks::Launch(UE_SOURCE_LOCATION, [WeakCapture, CapturingFrame, SyncDataPtr]()
+							UE::Tasks::FTask Task = UE::Tasks::Launch(UE_SOURCE_LOCATION, [MediaCapture, CapturingFrame, SyncDataPtr]()
 								{
 									TRACE_CPUPROFILER_EVENT_SCOPE(UMediaCapture::SyncPass);
 
@@ -739,25 +736,10 @@ namespace UE::MediaCaptureData
 												break;
 											}
 
-											// Can't check for object validity during garbage collection.
-											if (!IsGarbageCollecting())
+											if (!CapturingFrame->bMediaCaptureActive)
 											{
-
-												constexpr bool bEvenIfPendingKill = false;
-												constexpr bool ThreadSafeTest = true;
-												if (WeakCapture.IsValid(bEvenIfPendingKill, ThreadSafeTest))
-												{
-													if (!WeakCapture->bIsActive)
-													{
-														bWaitedForCompletion = false;
-														break;
-													}
-												}
-												else
-												{
-													bWaitedForCompletion = false;
-													break;
-												}
+												bWaitedForCompletion = false;
+												break;
 											}
 
 											FPlatformProcess::SleepNoStats(0);
@@ -767,22 +749,19 @@ namespace UE::MediaCaptureData
 										SyncDataPtr->bIsBusy = false;
 									}
 
-									if (WeakCapture->bIsActive && bWaitedForCompletion)
+									if (CapturingFrame->bMediaCaptureActive && bWaitedForCompletion)
 									{
 										if (CVarMediaIOScheduleOnAnyThread.GetValueOnAnyThread() == 1)
 										{
-											OnReadbackComplete(FRHICommandListExecutor::GetImmediateCommandList(), WeakCapture.Get(), CapturingFrame);
+											OnReadbackComplete(FRHICommandListExecutor::GetImmediateCommandList(), MediaCapture, CapturingFrame);
 										}
 										else
 										{
-											++WeakCapture->WaitingForRenderCommandExecutionCounter;
+											++MediaCapture->WaitingForRenderCommandExecutionCounter;
 
-											ENQUEUE_RENDER_COMMAND(MediaOutputCaptureFrameCreateResources)([WeakCapture, CapturingFrame](FRHICommandList& RHICommandList)
+											ENQUEUE_RENDER_COMMAND(MediaOutputCaptureFrameCreateResources)([MediaCapture, CapturingFrame](FRHICommandList& RHICommandList)
 											{
-												if (WeakCapture.IsValid(false, true))
-												{
-													OnReadbackComplete(RHICommandList, WeakCapture.Get(), CapturingFrame);
-												}
+												OnReadbackComplete(RHICommandList, MediaCapture, CapturingFrame);
 											});
 										}
 									}
@@ -791,8 +770,8 @@ namespace UE::MediaCaptureData
 
 							{
 								// Add the task to the capture's list of pending tasks.
-								FScopeLock ScopedLock(&WeakCapture->PendingReadbackTasksCriticalSection);
-								WeakCapture->PendingReadbackTasks.Add(MoveTemp(Task));
+								FScopeLock ScopedLock(&MediaCapture->PendingReadbackTasksCriticalSection);
+								MediaCapture->PendingReadbackTasks.Add(MoveTemp(Task));
 							}
 						}
 					}
@@ -1138,8 +1117,6 @@ bool UMediaCapture::StartSourceCapture(TSharedPtr<UE::MediaCapture::Private::FCa
 	MediaCaptureAnalytics::SendCaptureEvent(GetCaptureSourceType());
 #endif
 
-	bIsActive = true;
-
 	return bInitialized;
 }
 
@@ -1303,7 +1280,13 @@ void UMediaCapture::StopCapture(bool bAllowPendingFrameToBeProcess)
 	}
 	else
 	{
-		bIsActive = false;
+		if (FrameManager)
+		{
+			FrameManager->ForEachFrame([](const TSharedPtr<UE::MediaCaptureData::FFrame> InFrame)
+			{
+				StaticCastSharedPtr<UE::MediaCaptureData::FCaptureFrame>(InFrame)->bMediaCaptureActive = false;
+			});
+		}
 
 		if (GetState() != EMediaCaptureState::Stopped)
 		{
@@ -1470,13 +1453,19 @@ void UMediaCapture::CaptureImmediate_RenderThread(const UE::MediaCaptureData::FC
 		else
 		{
 			//In case we are skipping frames, just keep capture frame as invalid
-			UE_LOG(LogMediaIOCore, Verbose, TEXT("[%s] - No frames available for capture. Skipping"), *MediaOutputName);
+			UE_LOG(LogMediaIOCore, Warning, TEXT("[%s] - No frames available for capture. Skipping"), *MediaOutputName);
 		}
 	}
 
 	PrintFrameState();
 
-	UE_LOG(LogMediaIOCore, Verbose, TEXT("[%s - %s] - Capturing frame %d"), *MediaOutputName, *FThreadManager::GetThreadName(FPlatformTLS::GetCurrentThreadId()), CapturingFrame ? CapturingFrame->FrameId : -1);
+	if (!CapturingFrame || CapturingFrame->FrameId == -1)
+	{
+		TRACE_CPUPROFILER_EVENT_SCOPE(UMediaCapture::Capture_Invalid);
+	}
+	
+	UE_LOG(LogMediaIOCore, Warning, TEXT("[%s - %s] - Capturing frame %d"), *MediaOutputName, *FThreadManager::GetThreadName(FPlatformTLS::GetCurrentThreadId()), CapturingFrame ? CapturingFrame->FrameId : -1);
+	
 	if (CapturingFrame)
 	{
 		// Prepare frame to capture
diff --git a/Engine/Plugins/Media/MediaIOFramework/Source/MediaIOCore/Public/MediaCapture.h b/Engine/Plugins/Media/MediaIOFramework/Source/MediaIOCore/Public/MediaCapture.h
--- a/Engine/Plugins/Media/MediaIOFramework/Source/MediaIOCore/Public/MediaCapture.h
+++ b/Engine/Plugins/Media/MediaIOFramework/Source/MediaIOCore/Public/MediaCapture.h
@@ -612,9 +612,6 @@ private:
 	/** Array of sync handlers (fence) to sync when captured buffer is completed */
 	TArray<TSharedPtr<FMediaCaptureSyncData>> SyncHandlers;
 
-	/** Whether capture is active. Used to be queried by sync task */
-	std::atomic<bool> bIsActive = false;
-
 	FCriticalSection PendingReadbackTasksCriticalSection;
 
 	/** List of pending readback tasks, when CVarMediaIOEnableExperimentalScheduling and CVarMediaIOScheduleOnAnyThread are set to true. */
