Media Capture - Remove check that the number of capture frame must be >= 2
Change default texture number to 3 in order to avoid skipping in captureimmediate

diff --git a/Engine/Plugins/Media/MediaIOFramework/Source/MediaIOCore/Private/MediaCapture.cpp b/Engine/Plugins/Media/MediaIOFramework/Source/MediaIOCore/Private/MediaCapture.cpp
--- a/Engine/Plugins/Media/MediaIOFramework/Source/MediaIOCore/Private/MediaCapture.cpp
+++ b/Engine/Plugins/Media/MediaIOFramework/Source/MediaIOCore/Private/MediaCapture.cpp
@@ -1804,8 +1804,6 @@ void UMediaCapture::InitializeOutputResources(int32 InNumberOfBuffers)
 	bSyncHandlersInitialized = false;
 	NumberOfCaptureFrame = InNumberOfBuffers;
 
-	check(NumberOfCaptureFrame >= 2);
-
 	UMediaCapture* This = this;
 	ENQUEUE_RENDER_COMMAND(MediaOutputCaptureFrameCreateResources)(
 	[This](FRHICommandListImmediate& RHICmdList)
diff --git a/Engine/Plugins/Media/MediaIOFramework/Source/MediaIOCore/Private/MediaOutput.cpp b/Engine/Plugins/Media/MediaIOFramework/Source/MediaIOCore/Private/MediaOutput.cpp
--- a/Engine/Plugins/Media/MediaIOFramework/Source/MediaIOCore/Private/MediaOutput.cpp
+++ b/Engine/Plugins/Media/MediaIOFramework/Source/MediaIOCore/Private/MediaOutput.cpp
@@ -14,7 +14,7 @@ const FIntPoint UMediaOutput::RequestCaptureSourceSize = FIntPoint::ZeroValue;
 
 UMediaOutput::UMediaOutput(const FObjectInitializer& ObjectInitializer)
 	: Super(ObjectInitializer)
-	, NumberOfTextureBuffers(2)
+	, NumberOfTextureBuffers(3)
 {
 }
 
