Fix media capture not working if it doesn't support anythread capture

diff --git a/Engine/Plugins/Media/AjaMedia/Source/AjaMediaOutput/Public/AjaMediaCapture.h b/Engine/Plugins/Media/AjaMedia/Source/AjaMediaOutput/Public/AjaMediaCapture.h
--- a/Engine/Plugins/Media/AjaMedia/Source/AjaMediaOutput/Public/AjaMediaCapture.h
+++ b/Engine/Plugins/Media/AjaMedia/Source/AjaMediaOutput/Public/AjaMediaCapture.h
@@ -56,6 +56,11 @@ protected:
 	virtual void LockDMATexture_RenderThread(FTextureRHIRef InTexture) override;
 	virtual void UnlockDMATexture_RenderThread(FTextureRHIRef InTexture) override;
 	virtual void OnFrameCaptured_AnyThread(const FCaptureBaseData& InBaseData, TSharedPtr<FMediaCaptureUserData, ESPMode::ThreadSafe> InUserData, const FMediaCaptureResourceData& InResourceData) override;
+	virtual bool SupportsAnyThreadCapture() const override
+	{
+		return true;
+	}
+
 
 
 private:
diff --git a/Engine/Plugins/Media/BlackmagicMedia/Source/BlackmagicMediaOutput/Public/BlackmagicMediaCapture.h b/Engine/Plugins/Media/BlackmagicMedia/Source/BlackmagicMediaOutput/Public/BlackmagicMediaCapture.h
--- a/Engine/Plugins/Media/BlackmagicMedia/Source/BlackmagicMediaOutput/Public/BlackmagicMediaCapture.h
+++ b/Engine/Plugins/Media/BlackmagicMedia/Source/BlackmagicMediaOutput/Public/BlackmagicMediaCapture.h
@@ -57,6 +57,10 @@ protected:
 	virtual void OnRHIResourceCaptured_RenderingThread(const FCaptureBaseData& InBaseData, TSharedPtr<FMediaCaptureUserData, ESPMode::ThreadSafe> InUserData, FTextureRHIRef InTexture) override;
 	virtual void LockDMATexture_RenderThread(FTextureRHIRef InTexture) override;
 	virtual void UnlockDMATexture_RenderThread(FTextureRHIRef InTexture) override;
+	virtual bool SupportsAnyThreadCapture() const override
+	{
+		return true;
+	}
 
 private:
 	struct FBlackmagicOutputCallback;
diff --git a/Engine/Plugins/Media/MediaIOFramework/Source/MediaIOCore/Private/MediaCapture.cpp b/Engine/Plugins/Media/MediaIOFramework/Source/MediaIOCore/Private/MediaCapture.cpp
--- a/Engine/Plugins/Media/MediaIOFramework/Source/MediaIOCore/Private/MediaCapture.cpp
+++ b/Engine/Plugins/Media/MediaIOFramework/Source/MediaIOCore/Private/MediaCapture.cpp
@@ -751,7 +751,7 @@ namespace UE::MediaCaptureData
 
 									if (CapturingFrame->bMediaCaptureActive && bWaitedForCompletion)
 									{
-										if (CVarMediaIOScheduleOnAnyThread.GetValueOnAnyThread() == 1)
+										if (MediaCapture->UseAnyThreadCapture())
 										{
 											OnReadbackComplete(FRHICommandListExecutor::GetImmediateCommandList(), MediaCapture, CapturingFrame);
 										}
@@ -1117,6 +1117,15 @@ bool UMediaCapture::StartSourceCapture(TSharedPtr<UE::MediaCapture::Private::FCa
 	MediaCaptureAnalytics::SendCaptureEvent(GetCaptureSourceType());
 #endif
 
+	if (!GRHISupportsMultithreading)
+	{
+		UE_LOG(LogMediaIOCore, Display, TEXT("AnyThread Capture was disabled because the current RHI does not support Multithreading."));
+	}
+	else if (SupportsAnyThreadCapture())
+	{
+		UE_LOG(LogMediaIOCore, Display, TEXT("AnyThread Capture was disabled because the media capture implementation does not have a AnyThread callback."));
+	}
+
 	return bInitialized;
 }
 
@@ -1244,6 +1253,14 @@ bool UMediaCapture::UpdateSource(TSharedPtr<UE::MediaCapture::Private::FCaptureS
 	return true;
 }
 
+bool UMediaCapture::UseAnyThreadCapture() const
+{
+	return GRHISupportsMultithreading
+		&& CVarMediaIOEnableExperimentalScheduling.GetValueOnAnyThread() 
+		&& CVarMediaIOScheduleOnAnyThread.GetValueOnAnyThread()
+		&& SupportsAnyThreadCapture();
+}
+
 bool UMediaCapture::UpdateSceneViewport(TSharedPtr<FSceneViewport>& InSceneViewport)
 {
 	check(CaptureSource);
@@ -1587,7 +1604,7 @@ void UMediaCapture::WaitForSingleExperimentalSchedulingTaskToComplete()
 			}
 		}
 
-		if (!CVarMediaIOScheduleOnAnyThread.GetValueOnAnyThread())
+		if (!UseAnyThreadCapture())
 		{
 			// We flush after task completion in case a render thread task was launched.
 			if (IsInGameThread())
@@ -1632,7 +1649,7 @@ void UMediaCapture::WaitForAllExperimentalSchedulingTasksToComplete()
 		}
 
 		
-		if (!CVarMediaIOScheduleOnAnyThread.GetValueOnGameThread())
+		if (!UseAnyThreadCapture())
 		{
 			// This code might have dispatched a task on the render thread, so we need to wait again
 			if (IsInGameThread())
diff --git a/Engine/Plugins/Media/MediaIOFramework/Source/MediaIOCore/Public/MediaCapture.h b/Engine/Plugins/Media/MediaIOFramework/Source/MediaIOCore/Public/MediaCapture.h
--- a/Engine/Plugins/Media/MediaIOFramework/Source/MediaIOCore/Public/MediaCapture.h
+++ b/Engine/Plugins/Media/MediaIOFramework/Source/MediaIOCore/Public/MediaCapture.h
@@ -376,6 +376,12 @@ protected:
 	virtual void LockDMATexture_RenderThread(FTextureRHIRef InTexture) {}
 	virtual void UnlockDMATexture_RenderThread(FTextureRHIRef InTexture) {}
 
+	/** Whether the capture callbacks can be called on any thread. If true, the _AnyThread callbacks will be used instead of the _RenderThread ones. */
+	virtual bool SupportsAnyThreadCapture() const
+	{
+		return false;
+	}
+
 	friend class UE::MediaCaptureData::FCaptureFrame;
 	friend class UE::MediaCaptureData::FMediaCaptureHelper;
 	struct FCaptureBaseData
@@ -425,6 +431,7 @@ protected:
 	/**
 	 * Callback when the buffer was successfully copied to CPU ram.
 	 * The buffer is only valid for the duration of the callback.
+	 * @Note SupportsAnyThreadCapture must return true in the implementation in order for this to be called.
 	 */
 	virtual void OnFrameCaptured_AnyThread(const FCaptureBaseData& InBaseData, TSharedPtr<FMediaCaptureUserData, ESPMode::ThreadSafe> InUserData, const FMediaCaptureResourceData& InResourceData) { }
 
@@ -439,6 +446,10 @@ protected:
 	virtual void OnRHIResourceCaptured_RenderingThread(const FCaptureBaseData& InBaseData, TSharedPtr<FMediaCaptureUserData, ESPMode::ThreadSafe> InUserData, FTextureRHIRef InTexture) { }
 	virtual void OnRHIResourceCaptured_RenderingThread(const FCaptureBaseData& InBaseData, TSharedPtr<FMediaCaptureUserData, ESPMode::ThreadSafe> InUserData, FBufferRHIRef InBuffer) { }
 	
+	/**
+	 * AnyThread version of the above callbacks.
+	 * @Note SupportsAnyThreadCapture must return true in the implementation in order for this to be called.
+	 */
 	virtual void OnRHIResourceCaptured_AnyThread(const FCaptureBaseData& InBaseData, TSharedPtr<FMediaCaptureUserData, ESPMode::ThreadSafe> InUserData, FTextureRHIRef InTexture) { }
 	virtual void OnRHIResourceCaptured_AnyThread(const FCaptureBaseData& InBaseData, TSharedPtr<FMediaCaptureUserData, ESPMode::ThreadSafe> InUserData, FBufferRHIRef InBuffer) { }
 	
@@ -544,6 +555,8 @@ private:
 	bool StartSourceCapture(TSharedPtr<UE::MediaCapture::Private::FCaptureSource> InSource);
 	bool UpdateSource(TSharedPtr<UE::MediaCapture::Private::FCaptureSource> InCaptureSource);
 
+	bool UseAnyThreadCapture() const;
+
 protected:
 
 	/** MediaOutput associated with this capture */
