Change a log verbosity in mediacapture

diff --git a/Engine/Plugins/Media/MediaIOFramework/Source/MediaIOCore/Private/MediaCapture.cpp b/Engine/Plugins/Media/MediaIOFramework/Source/MediaIOCore/Private/MediaCapture.cpp
--- a/Engine/Plugins/Media/MediaIOFramework/Source/MediaIOCore/Private/MediaCapture.cpp
+++ b/Engine/Plugins/Media/MediaIOFramework/Source/MediaIOCore/Private/MediaCapture.cpp
@@ -1481,7 +1481,7 @@ void UMediaCapture::CaptureImmediate_RenderThread(const UE::MediaCaptureData::FC
 		TRACE_CPUPROFILER_EVENT_SCOPE(UMediaCapture::Capture_Invalid);
 	}
 	
-	UE_LOG(LogMediaIOCore, Warning, TEXT("[%s - %s] - Capturing frame %d"), *MediaOutputName, *FThreadManager::GetThreadName(FPlatformTLS::GetCurrentThreadId()), CapturingFrame ? CapturingFrame->FrameId : -1);
+	UE_LOG(LogMediaIOCore, Verbose, TEXT("[%s - %s] - Capturing frame %d"), *MediaOutputName, *FThreadManager::GetThreadName(FPlatformTLS::GetCurrentThreadId()), CapturingFrame ? CapturingFrame->FrameId : -1);
 	
 	if (CapturingFrame)
 	{
