Adding includes for upcoming dependency cleanup.

diff --git a/Engine/Plugins/AudioGameplayVolume/Source/AudioGameplayVolume/Private/AudioGameplayVolumeComponent.cpp b/Engine/Plugins/AudioGameplayVolume/Source/AudioGameplayVolume/Private/AudioGameplayVolumeComponent.cpp
--- a/Engine/Plugins/AudioGameplayVolume/Source/AudioGameplayVolume/Private/AudioGameplayVolumeComponent.cpp
+++ b/Engine/Plugins/AudioGameplayVolume/Source/AudioGameplayVolume/Private/AudioGameplayVolumeComponent.cpp
@@ -3,6 +3,7 @@
 #include "AudioGameplayVolumeComponent.h"
 #include "AudioGameplayVolumeSubsystem.h"
 #include "AudioDevice.h"
+#include "Engine/World.h"
 
 #include UE_INLINE_GENERATED_CPP_BY_NAME(AudioGameplayVolumeComponent)
 
diff --git a/Engine/Plugins/AudioGameplayVolume/Source/AudioGameplayVolume/Private/AudioGameplayVolumeSubsystem.cpp b/Engine/Plugins/AudioGameplayVolume/Source/AudioGameplayVolume/Private/AudioGameplayVolumeSubsystem.cpp
--- a/Engine/Plugins/AudioGameplayVolume/Source/AudioGameplayVolume/Private/AudioGameplayVolumeSubsystem.cpp
+++ b/Engine/Plugins/AudioGameplayVolume/Source/AudioGameplayVolume/Private/AudioGameplayVolumeSubsystem.cpp
@@ -7,6 +7,7 @@
 #include "AudioGameplayVolumeProxy.h"
 #include "AudioGameplayVolumeMutator.h"
 #include "AudioDevice.h"
+#include "Engine/World.h"
 #include "Misc/App.h"
 
 #include UE_INLINE_GENERATED_CPP_BY_NAME(AudioGameplayVolumeSubsystem)
diff --git a/Engine/Plugins/Media/BinkMedia/Source/BinkMediaPlayer/Private/BinkMediaPlayerPrivate.h b/Engine/Plugins/Media/BinkMedia/Source/BinkMediaPlayer/Private/BinkMediaPlayerPrivate.h
--- a/Engine/Plugins/Media/BinkMedia/Source/BinkMediaPlayer/Private/BinkMediaPlayerPrivate.h
+++ b/Engine/Plugins/Media/BinkMedia/Source/BinkMediaPlayer/Private/BinkMediaPlayerPrivate.h
@@ -2,6 +2,7 @@
 //   Licenced under the Unreal Engine EULA 
 #pragma once
 
+#include "Engine/Engine.h"
 #include "Rendering/SlateRenderer.h"
 #include "AudioMixerDevice.h"
 
diff --git a/Engine/Plugins/Media/MediaIOFramework/Source/MediaIOCore/Public/MediaIOCoreAudioOutput.h b/Engine/Plugins/Media/MediaIOFramework/Source/MediaIOCore/Public/MediaIOCoreAudioOutput.h
--- a/Engine/Plugins/Media/MediaIOFramework/Source/MediaIOCore/Public/MediaIOCoreAudioOutput.h
+++ b/Engine/Plugins/Media/MediaIOFramework/Source/MediaIOCore/Public/MediaIOCoreAudioOutput.h
@@ -3,14 +3,17 @@
 #pragma once
 
 #include "DSP/BufferVectorOperations.h"
+#include "DSP/MultithreadedPatching.h"
 #include "ISubmixBufferListener.h"
 #include "Math/NumericLimits.h"
+#include "Misc/FrameRate.h"
 
 #if WITH_MEDIA_IO_AUDIO_DEBUGGING
 #include "MediaIOAudioDebug.h"
 #endif
 
 class FAudioDevice;
+
 class FAudioDeviceHandle;
 
 DECLARE_DELEGATE_TwoParams(FOnBufferReceived, uint8* Buffer, int32 BufferSize);
diff --git a/Engine/Plugins/Media/PixelStreaming/Source/PixelStreaming/Private/AudioCapturer.h b/Engine/Plugins/Media/PixelStreaming/Source/PixelStreaming/Private/AudioCapturer.h
--- a/Engine/Plugins/Media/PixelStreaming/Source/PixelStreaming/Private/AudioCapturer.h
+++ b/Engine/Plugins/Media/PixelStreaming/Source/PixelStreaming/Private/AudioCapturer.h
@@ -4,6 +4,7 @@
 
 #include "WebRTCIncludes.h"
 #include "HAL/CriticalSection.h"
+#include "HAL/ThreadSafeBool.h"
 #include "ISubmixBufferListener.h"
 
 namespace UE::PixelStreaming 
diff --git a/Engine/Plugins/Media/PixelStreaming/Source/PixelStreaming/Private/AudioSubmixCapturer.cpp b/Engine/Plugins/Media/PixelStreaming/Source/PixelStreaming/Private/AudioSubmixCapturer.cpp
--- a/Engine/Plugins/Media/PixelStreaming/Source/PixelStreaming/Private/AudioSubmixCapturer.cpp
+++ b/Engine/Plugins/Media/PixelStreaming/Source/PixelStreaming/Private/AudioSubmixCapturer.cpp
@@ -4,7 +4,9 @@
 
 #include "AudioDevice.h"
 #include "AudioInputMixer.h"
+#include "Engine/Engine.h"
 #include "PixelStreamingPrivate.h"
+#include "SampleBuffer.h"
 #include "Settings.h"
 
 namespace UE::PixelStreaming
diff --git a/Engine/Plugins/Media/PixelStreaming/Source/PixelStreaming/Private/AudioSubmixCapturer.h b/Engine/Plugins/Media/PixelStreaming/Source/PixelStreaming/Private/AudioSubmixCapturer.h
--- a/Engine/Plugins/Media/PixelStreaming/Source/PixelStreaming/Private/AudioSubmixCapturer.h
+++ b/Engine/Plugins/Media/PixelStreaming/Source/PixelStreaming/Private/AudioSubmixCapturer.h
@@ -2,7 +2,9 @@
 
 #pragma once
 
+#include "HAL/ThreadSafeBool.h"
 #include "ISubmixBufferListener.h"
+#include "Templates/SharedPointer.h"
 
 namespace webrtc
 {
diff --git a/Engine/Plugins/Runtime/Metasound/Source/MetasoundEngineTest/Private/EngineTestMetasoundSource.cpp b/Engine/Plugins/Runtime/Metasound/Source/MetasoundEngineTest/Private/EngineTestMetasoundSource.cpp
--- a/Engine/Plugins/Runtime/Metasound/Source/MetasoundEngineTest/Private/EngineTestMetasoundSource.cpp
+++ b/Engine/Plugins/Runtime/Metasound/Source/MetasoundEngineTest/Private/EngineTestMetasoundSource.cpp
@@ -10,6 +10,7 @@
 #include "MetasoundFrontendSearchEngine.h"
 #include "MetasoundSource.h"
 #include "Misc/AutomationTest.h"
+#include "Misc/Paths.h"
 #include "Tests/AutomationCommon.h"
 
 #if WITH_DEV_AUTOMATION_TESTS
diff --git a/Engine/Plugins/Runtime/Oculus/OculusAudio/Source/OculusAudio/Private/OculusAudioGeometryComponent.cpp b/Engine/Plugins/Runtime/Oculus/OculusAudio/Source/OculusAudio/Private/OculusAudioGeometryComponent.cpp
--- a/Engine/Plugins/Runtime/Oculus/OculusAudio/Source/OculusAudio/Private/OculusAudioGeometryComponent.cpp
+++ b/Engine/Plugins/Runtime/Oculus/OculusAudio/Source/OculusAudio/Private/OculusAudioGeometryComponent.cpp
@@ -4,6 +4,7 @@
 #include "Components/StaticMeshComponent.h"
 #include "Engine/StaticMesh.h"
 #include "Engine/StaticMeshActor.h"
+#include "Engine/World.h"
 #include "IOculusAudioPlugin.h"
 #include "Materials/MaterialInstanceDynamic.h"
 #include "OculusAudio.h"
diff --git a/Engine/Plugins/Runtime/Oculus/OculusAudio/Source/OculusAudio/Private/OculusAudioReverb.cpp b/Engine/Plugins/Runtime/Oculus/OculusAudio/Source/OculusAudio/Private/OculusAudioReverb.cpp
--- a/Engine/Plugins/Runtime/Oculus/OculusAudio/Source/OculusAudio/Private/OculusAudioReverb.cpp
+++ b/Engine/Plugins/Runtime/Oculus/OculusAudio/Source/OculusAudio/Private/OculusAudioReverb.cpp
@@ -6,6 +6,7 @@
 #include "OculusAudioContextManager.h"
 
 #include "Sound/SoundSubmix.h"
+#include "Engine/Engine.h"
 
 
 namespace
diff --git a/Engine/Plugins/Runtime/ResonanceAudio/Source/ResonanceAudio/Private/ResonanceAudioPluginListener.cpp b/Engine/Plugins/Runtime/ResonanceAudio/Source/ResonanceAudio/Private/ResonanceAudioPluginListener.cpp
--- a/Engine/Plugins/Runtime/ResonanceAudio/Source/ResonanceAudio/Private/ResonanceAudioPluginListener.cpp
+++ b/Engine/Plugins/Runtime/ResonanceAudio/Source/ResonanceAudio/Private/ResonanceAudioPluginListener.cpp
@@ -7,6 +7,7 @@
 #include "AudioDevice.h"
 #include "Async/Async.h"
 #include "Components/BrushComponent.h"
+#include "Engine/World.h"
 #include "Model.h"
 
 namespace ResonanceAudioConsoleVariables
diff --git a/Engine/Plugins/Runtime/Soundscape/Source/SoundScape/Private/SoundscapeSubsystem.cpp b/Engine/Plugins/Runtime/Soundscape/Source/SoundScape/Private/SoundscapeSubsystem.cpp
--- a/Engine/Plugins/Runtime/Soundscape/Source/SoundScape/Private/SoundscapeSubsystem.cpp
+++ b/Engine/Plugins/Runtime/Soundscape/Source/SoundScape/Private/SoundscapeSubsystem.cpp
@@ -4,6 +4,7 @@
 #include "SoundscapeSubsystem.h"
 #include "DrawDebugHelpers.h"
 #include "Engine/GameInstance.h"
+#include "Engine/World.h"
 #include "SoundscapeColorPoint.h"
 #include "SoundscapeSettings.h"
 #include "SoundScapePalette.h"
diff --git a/Engine/Plugins/Runtime/Steam/SteamAudio/Source/SteamAudio/Private/PhononGeometryComponent.cpp b/Engine/Plugins/Runtime/Steam/SteamAudio/Source/SteamAudio/Private/PhononGeometryComponent.cpp
--- a/Engine/Plugins/Runtime/Steam/SteamAudio/Source/SteamAudio/Private/PhononGeometryComponent.cpp
+++ b/Engine/Plugins/Runtime/Steam/SteamAudio/Source/SteamAudio/Private/PhononGeometryComponent.cpp
@@ -3,6 +3,7 @@
 //
 
 #include "PhononGeometryComponent.h"
+#include "Engine/Engine.h"
 #include "PhononScene.h"
 #include "StaticMeshResources.h"
 #include "PhononPluginManager.h"
diff --git a/Engine/Plugins/Runtime/Steam/SteamAudio/Source/SteamAudio/Private/PhononPluginManager.cpp b/Engine/Plugins/Runtime/Steam/SteamAudio/Source/SteamAudio/Private/PhononPluginManager.cpp
--- a/Engine/Plugins/Runtime/Steam/SteamAudio/Source/SteamAudio/Private/PhononPluginManager.cpp
+++ b/Engine/Plugins/Runtime/Steam/SteamAudio/Source/SteamAudio/Private/PhononPluginManager.cpp
@@ -7,6 +7,7 @@
 #include "AudioDevice.h"
 #include "HAL/FileManager.h"
 #include "DrawDebugHelpers.h"
+#include "Engine/World.h"
 
 namespace SteamAudio
 {
diff --git a/Engine/Plugins/Runtime/Windows/WinDualShock/Source/WinDualShock/Private/WinDualShock.cpp b/Engine/Plugins/Runtime/Windows/WinDualShock/Source/WinDualShock/Private/WinDualShock.cpp
--- a/Engine/Plugins/Runtime/Windows/WinDualShock/Source/WinDualShock/Private/WinDualShock.cpp
+++ b/Engine/Plugins/Runtime/Windows/WinDualShock/Source/WinDualShock/Private/WinDualShock.cpp
@@ -2,6 +2,7 @@
 
 #include "WinDualShock.h"
 #include "Containers/CircularQueue.h"
+#include "Engine/Engine.h"
 #include "Misc/CoreDelegates.h"
 #include "Misc/MessageDialog.h"
 #include "Modules/ModuleManager.h"
diff --git a/Engine/Source/Editor/DetailCustomizations/Private/LinuxTargetSettingsDetails.cpp b/Engine/Source/Editor/DetailCustomizations/Private/LinuxTargetSettingsDetails.cpp
--- a/Engine/Source/Editor/DetailCustomizations/Private/LinuxTargetSettingsDetails.cpp
+++ b/Engine/Source/Editor/DetailCustomizations/Private/LinuxTargetSettingsDetails.cpp
@@ -1,6 +1,7 @@
 // Copyright Epic Games, Inc. All Rights Reserved.
 
 #include "LinuxTargetSettingsDetails.h"
+#include "Engine/Engine.h"
 #include "Misc/Paths.h"
 #include "Misc/ConfigCacheIni.h"
 #include "Misc/App.h"
diff --git a/Engine/Source/Editor/UnrealEd/Private/Commandlets/AudioMixerCommandlet.cpp b/Engine/Source/Editor/UnrealEd/Private/Commandlets/AudioMixerCommandlet.cpp
--- a/Engine/Source/Editor/UnrealEd/Private/Commandlets/AudioMixerCommandlet.cpp
+++ b/Engine/Source/Editor/UnrealEd/Private/Commandlets/AudioMixerCommandlet.cpp
@@ -3,6 +3,7 @@
 #include "Commandlets/AudioMixerCommandlet.h"
 #include "HAL/FileManager.h"
 #include "Misc/Paths.h"
+#include "Engine/Engine.h"
 #include "Engine/EngineBaseTypes.h"
 #include "Sound/SoundAttenuation.h"
 #include "Sound/SoundWave.h"
diff --git a/Engine/Source/Runtime/AudioAnalyzer/Private/AudioAnalyzer.cpp b/Engine/Source/Runtime/AudioAnalyzer/Private/AudioAnalyzer.cpp
--- a/Engine/Source/Runtime/AudioAnalyzer/Private/AudioAnalyzer.cpp
+++ b/Engine/Source/Runtime/AudioAnalyzer/Private/AudioAnalyzer.cpp
@@ -9,6 +9,7 @@
 #include "AudioDeviceManager.h"
 #include "AudioAnalyzerSubsystem.h"
 #include "AudioBusSubsystem.h"
+#include "Engine/Engine.h"
 
 #include UE_INLINE_GENERATED_CPP_BY_NAME(AudioAnalyzer)
 
diff --git a/Engine/Source/Runtime/AudioMixer/Private/AudioMixerSource.cpp b/Engine/Source/Runtime/AudioMixer/Private/AudioMixerSource.cpp
--- a/Engine/Source/Runtime/AudioMixer/Private/AudioMixerSource.cpp
+++ b/Engine/Source/Runtime/AudioMixer/Private/AudioMixerSource.cpp
@@ -17,6 +17,7 @@
 #include "Misc/ScopeRWLock.h"
 #include "Templates/Function.h"
 #include "Trace/Trace.h"
+#include "Engine/Engine.h"
 
 
 CSV_DECLARE_CATEGORY_MODULE_EXTERN(AUDIOMIXERCORE_API, Audio);
diff --git a/Engine/Source/Runtime/AudioMixer/Private/AudioMixerSubmix.cpp b/Engine/Source/Runtime/AudioMixer/Private/AudioMixerSubmix.cpp
--- a/Engine/Source/Runtime/AudioMixer/Private/AudioMixerSubmix.cpp
+++ b/Engine/Source/Runtime/AudioMixer/Private/AudioMixerSubmix.cpp
@@ -7,6 +7,7 @@
 #include "AudioMixerSourceVoice.h"
 #include "AudioThread.h"
 #include "DSP/FloatArrayMath.h"
+#include "ISubmixBufferListener.h"
 #include "Sound/SoundEffectPreset.h"
 #include "Sound/SoundEffectSubmix.h"
 #include "Sound/SoundModulationDestination.h"
diff --git a/Engine/Source/Runtime/AudioMixer/Private/Components/SynthComponent.cpp b/Engine/Source/Runtime/AudioMixer/Private/Components/SynthComponent.cpp
--- a/Engine/Source/Runtime/AudioMixer/Private/Components/SynthComponent.cpp
+++ b/Engine/Source/Runtime/AudioMixer/Private/Components/SynthComponent.cpp
@@ -4,6 +4,7 @@
 
 #include "AudioDevice.h"
 #include "AudioMixerLog.h"
+#include "Engine/World.h"
 #include "Sound/AudioSettings.h"
 
 #include UE_INLINE_GENERATED_CPP_BY_NAME(SynthComponent)
diff --git a/Engine/Source/Runtime/AudioMixer/Private/Quartz/QuartzSubsystem.cpp b/Engine/Source/Runtime/AudioMixer/Private/Quartz/QuartzSubsystem.cpp
--- a/Engine/Source/Runtime/AudioMixer/Private/Quartz/QuartzSubsystem.cpp
+++ b/Engine/Source/Runtime/AudioMixer/Private/Quartz/QuartzSubsystem.cpp
@@ -2,6 +2,7 @@
 
 #include "Quartz/QuartzSubsystem.h"
 
+#include "Engine/World.h"
 #include "Quartz/AudioMixerClockHandle.h"
 #include "Quartz/QuartzMetronome.h"
 #include "Quartz/AudioMixerClockManager.h"
diff --git a/Engine/Source/Runtime/AudioMixer/Public/AudioMixerDevice.h b/Engine/Source/Runtime/AudioMixer/Public/AudioMixerDevice.h
--- a/Engine/Source/Runtime/AudioMixer/Public/AudioMixerDevice.h
+++ b/Engine/Source/Runtime/AudioMixer/Public/AudioMixerDevice.h
@@ -5,6 +5,7 @@
 #include "Audio.h"
 #include "AudioMixer.h"
 #include "AudioDevice.h"
+#include "Containers/MpscQueue.h"
 #include "Sound/SoundSubmix.h"
 #include "Sound/SoundGenerator.h"
 #include "DSP/BufferVectorOperations.h"
diff --git a/Engine/Source/Runtime/Engine/Private/ActiveSound.cpp b/Engine/Source/Runtime/Engine/Private/ActiveSound.cpp
--- a/Engine/Source/Runtime/Engine/Private/ActiveSound.cpp
+++ b/Engine/Source/Runtime/Engine/Private/ActiveSound.cpp
@@ -6,6 +6,7 @@
 #include "AudioDevice.h"
 #include "AudioLinkSettingsAbstract.h"
 #include "Sound/SoundCue.h"
+#include "Engine/Engine.h"
 #include "Engine/World.h"
 #include "Sound/SoundNodeAttenuation.h"
 #include "IAudioParameterTransmitter.h"
diff --git a/Engine/Source/Runtime/Engine/Private/Audio.cpp b/Engine/Source/Runtime/Engine/Private/Audio.cpp
--- a/Engine/Source/Runtime/Engine/Private/Audio.cpp
+++ b/Engine/Source/Runtime/Engine/Private/Audio.cpp
@@ -10,6 +10,7 @@
 #include "AudioDevice.h"
 #include "AudioPluginUtilities.h"
 #include "Components/SynthComponent.h"
+#include "Engine/Engine.h"
 #include "EngineAnalytics.h"
 #include "IAnalyticsProviderET.h"
 #include "Misc/Paths.h"
diff --git a/Engine/Source/Runtime/Engine/Private/AudioBus.cpp b/Engine/Source/Runtime/Engine/Private/AudioBus.cpp
--- a/Engine/Source/Runtime/Engine/Private/AudioBus.cpp
+++ b/Engine/Source/Runtime/Engine/Private/AudioBus.cpp
@@ -4,6 +4,7 @@
 #include "Sound/AudioBus.h"
 #include "AudioBusSubsystem.h"
 #include "AudioMixerDevice.h"
+#include "Engine/Engine.h"
 
 #include UE_INLINE_GENERATED_CPP_BY_NAME(AudioBus)
 
diff --git a/Engine/Source/Runtime/Engine/Private/AudioDecompress.cpp b/Engine/Source/Runtime/Engine/Private/AudioDecompress.cpp
--- a/Engine/Source/Runtime/Engine/Private/AudioDecompress.cpp
+++ b/Engine/Source/Runtime/Engine/Private/AudioDecompress.cpp
@@ -3,6 +3,7 @@
 
 #include "AudioDecompress.h"
 #include "AudioDevice.h"
+#include "Engine/Engine.h"
 #include "Interfaces/IAudioFormat.h"
 #include "ADPCMAudioInfo.h"
 #include "Misc/CoreStats.h"
diff --git a/Engine/Source/Runtime/Engine/Private/AudioDeviceManager.cpp b/Engine/Source/Runtime/Engine/Private/AudioDeviceManager.cpp
--- a/Engine/Source/Runtime/Engine/Private/AudioDeviceManager.cpp
+++ b/Engine/Source/Runtime/Engine/Private/AudioDeviceManager.cpp
@@ -5,6 +5,7 @@
 #include "Audio/AudioDebug.h"
 #include "AudioAnalytics.h"
 #include "AudioMixerDevice.h"
+#include "Engine/Engine.h"
 #include "Engine/World.h"
 #include "Features/IModularFeatures.h"
 #include "Misc/App.h"
diff --git a/Engine/Source/Runtime/Engine/Private/AudioEffect.cpp b/Engine/Source/Runtime/Engine/Private/AudioEffect.cpp
--- a/Engine/Source/Runtime/Engine/Private/AudioEffect.cpp
+++ b/Engine/Source/Runtime/Engine/Private/AudioEffect.cpp
@@ -5,6 +5,7 @@
 =============================================================================*/
 
 #include "AudioEffect.h"
+#include "Engine/Engine.h"
 #include "Misc/App.h"
 #include "Sound/ReverbEffect.h"
 #include "AudioDevice.h"
diff --git a/Engine/Source/Runtime/Engine/Private/Components/AudioComponent.cpp b/Engine/Source/Runtime/Engine/Private/Components/AudioComponent.cpp
--- a/Engine/Source/Runtime/Engine/Private/Components/AudioComponent.cpp
+++ b/Engine/Source/Runtime/Engine/Private/Components/AudioComponent.cpp
@@ -4,6 +4,7 @@
 #include "Audio/ActorSoundParameterInterface.h"
 #include "AudioDevice.h"
 #include "Components/BillboardComponent.h"
+#include "Engine/Engine.h"
 #include "Engine/Texture2D.h"
 #include "Engine/World.h"
 #include "Kismet/GameplayStatics.h"
diff --git a/Engine/Source/Runtime/Engine/Private/ContentStreaming.cpp b/Engine/Source/Runtime/Engine/Private/ContentStreaming.cpp
--- a/Engine/Source/Runtime/Engine/Private/ContentStreaming.cpp
+++ b/Engine/Source/Runtime/Engine/Private/ContentStreaming.cpp
@@ -6,6 +6,7 @@
 
 #include "ContentStreaming.h"
 #include "Components/PrimitiveComponent.h"
+#include "Engine/Engine.h"
 #include "Engine/Texture2D.h"
 #include "Misc/ConfigCacheIni.h"
 #include "RHI.h"
diff --git a/Engine/Source/Runtime/Engine/Private/GameFramework/HapticFeedbackEffect.cpp b/Engine/Source/Runtime/Engine/Private/GameFramework/HapticFeedbackEffect.cpp
--- a/Engine/Source/Runtime/Engine/Private/GameFramework/HapticFeedbackEffect.cpp
+++ b/Engine/Source/Runtime/Engine/Private/GameFramework/HapticFeedbackEffect.cpp
@@ -2,6 +2,7 @@
 //
 
 #include "AudioDevice.h"
+#include "Engine/Engine.h"
 #include "GenericPlatform/IInputInterface.h"
 #include "Haptics/HapticFeedbackEffect_Base.h"
 #include "Haptics/HapticFeedbackEffect_Curve.h"
diff --git a/Engine/Source/Runtime/Engine/Private/GameplayStatics.cpp b/Engine/Source/Runtime/Engine/Private/GameplayStatics.cpp
--- a/Engine/Source/Runtime/Engine/Private/GameplayStatics.cpp
+++ b/Engine/Source/Runtime/Engine/Private/GameplayStatics.cpp
@@ -2,6 +2,7 @@
 
 #include "Kismet/GameplayStatics.h"
 #include "Engine/Blueprint.h"
+#include "Engine/Engine.h"
 #include "Engine/GameInstance.h"
 #include "Engine/GameViewportClient.h"
 #include "Serialization/ObjectAndNameAsStringProxyArchive.h"
diff --git a/Engine/Source/Runtime/Engine/Private/PlayerCameraManager.cpp b/Engine/Source/Runtime/Engine/Private/PlayerCameraManager.cpp
--- a/Engine/Source/Runtime/Engine/Private/PlayerCameraManager.cpp
+++ b/Engine/Source/Runtime/Engine/Private/PlayerCameraManager.cpp
@@ -2,6 +2,7 @@
 
 #include "Camera/PlayerCameraManager.h"
 #include "GameFramework/Pawn.h"
+#include "Engine/Engine.h"
 #include "Engine/World.h"
 #include "Camera/CameraActor.h"
 #include "Engine/Canvas.h"
diff --git a/Engine/Source/Runtime/Engine/Private/PreviewScene.cpp b/Engine/Source/Runtime/Engine/Private/PreviewScene.cpp
--- a/Engine/Source/Runtime/Engine/Private/PreviewScene.cpp
+++ b/Engine/Source/Runtime/Engine/Private/PreviewScene.cpp
@@ -5,6 +5,7 @@
 =============================================================================*/
 
 #include "PreviewScene.h"
+#include "Engine/Engine.h"
 #include "Engine/GameInstance.h"
 #include "Misc/ConfigCacheIni.h"
 #include "UObject/Package.h"
diff --git a/Engine/Source/Runtime/Engine/Private/SampleBufferIO.cpp b/Engine/Source/Runtime/Engine/Private/SampleBufferIO.cpp
--- a/Engine/Source/Runtime/Engine/Private/SampleBufferIO.cpp
+++ b/Engine/Source/Runtime/Engine/Private/SampleBufferIO.cpp
@@ -1,6 +1,7 @@
 // Copyright Epic Games, Inc. All Rights Reserved.
 
 #include "Sound/SampleBufferIO.h"
+#include "Engine/Engine.h"
 #include "HAL/PlatformFile.h"
 #include "HAL/PlatformFileManager.h"
 #include "AssetRegistry/AssetRegistryModule.h"
diff --git a/Engine/Source/Runtime/Engine/Private/Slate/SlateSoundDevice.cpp b/Engine/Source/Runtime/Engine/Private/Slate/SlateSoundDevice.cpp
--- a/Engine/Source/Runtime/Engine/Private/Slate/SlateSoundDevice.cpp
+++ b/Engine/Source/Runtime/Engine/Private/Slate/SlateSoundDevice.cpp
@@ -3,6 +3,7 @@
 #include "Slate/SlateSoundDevice.h"
 #include "Sound/SlateSound.h"
 #include "AudioDevice.h"
+#include "Engine/Engine.h"
 
 DEFINE_LOG_CATEGORY_STATIC(LogSlateSoundDevice, Log, All);
 
diff --git a/Engine/Source/Runtime/Engine/Private/SoundClass.cpp b/Engine/Source/Runtime/Engine/Private/SoundClass.cpp
--- a/Engine/Source/Runtime/Engine/Private/SoundClass.cpp
+++ b/Engine/Source/Runtime/Engine/Private/SoundClass.cpp
@@ -2,6 +2,7 @@
 
 
 #include "Sound/SoundClass.h"
+#include "Engine/Engine.h"
 #include "Sound/AudioOutputTarget.h"
 #include "UObject/UObjectIterator.h"
 #include "AudioDevice.h"
diff --git a/Engine/Source/Runtime/Engine/Private/SoundCue.cpp b/Engine/Source/Runtime/Engine/Private/SoundCue.cpp
--- a/Engine/Source/Runtime/Engine/Private/SoundCue.cpp
+++ b/Engine/Source/Runtime/Engine/Private/SoundCue.cpp
@@ -2,6 +2,7 @@
 
 #include "Sound/SoundCue.h"
 #include "EdGraph/EdGraph.h"
+#include "Engine/Engine.h"
 #include "Misc/CoreDelegates.h"
 #include "EdGraph/EdGraphSchema.h"
 #include "UObject/UObjectIterator.h"
diff --git a/Engine/Source/Runtime/Engine/Private/SoundSubmix.cpp b/Engine/Source/Runtime/Engine/Private/SoundSubmix.cpp
--- a/Engine/Source/Runtime/Engine/Private/SoundSubmix.cpp
+++ b/Engine/Source/Runtime/Engine/Private/SoundSubmix.cpp
@@ -2,6 +2,8 @@
 #include "Sound/SoundSubmix.h"
 
 #include "AudioDevice.h"
+#include "Engine/Engine.h"
+#include "Engine/World.h"
 #include "Sound/SampleBufferIO.h"
 
 #include UE_INLINE_GENERATED_CPP_BY_NAME(SoundSubmix)
diff --git a/Engine/Source/Runtime/Engine/Private/SoundWave.cpp b/Engine/Source/Runtime/Engine/Private/SoundWave.cpp
--- a/Engine/Source/Runtime/Engine/Private/SoundWave.cpp
+++ b/Engine/Source/Runtime/Engine/Private/SoundWave.cpp
@@ -6,6 +6,7 @@
 #include "AudioDecompress.h"
 #include "AudioDerivedData.h"
 #include "AudioDevice.h"
+#include "Engine/Engine.h"
 #include "Misc/Paths.h"
 #include "Serialization/MemoryWriter.h"
 #include "Sound/SoundSourceBus.h"
diff --git a/Engine/Source/Runtime/Engine/Private/SoundWaveProcedural.cpp b/Engine/Source/Runtime/Engine/Private/SoundWaveProcedural.cpp
--- a/Engine/Source/Runtime/Engine/Private/SoundWaveProcedural.cpp
+++ b/Engine/Source/Runtime/Engine/Private/SoundWaveProcedural.cpp
@@ -3,6 +3,7 @@
 #include "Sound/SoundWaveProcedural.h"
 
 #include "AudioDevice.h"
+#include "Engine/Engine.h"
 
 #include UE_INLINE_GENERATED_CPP_BY_NAME(SoundWaveProcedural)
 
diff --git a/Engine/Source/Runtime/Engine/Private/Subsystems/SubsystemBlueprintLibrary.cpp b/Engine/Source/Runtime/Engine/Private/Subsystems/SubsystemBlueprintLibrary.cpp
--- a/Engine/Source/Runtime/Engine/Private/Subsystems/SubsystemBlueprintLibrary.cpp
+++ b/Engine/Source/Runtime/Engine/Private/Subsystems/SubsystemBlueprintLibrary.cpp
@@ -2,6 +2,7 @@
 
 #include "Subsystems/SubsystemBlueprintLibrary.h"
 
+#include "Engine/Engine.h"
 #include "Engine/GameInstance.h"
 #include "AudioDevice.h"
 #include "Blueprint/UserWidget.h"
diff --git a/Engine/Source/Runtime/Engine/Private/WorldSettings.cpp b/Engine/Source/Runtime/Engine/Private/WorldSettings.cpp
--- a/Engine/Source/Runtime/Engine/Private/WorldSettings.cpp
+++ b/Engine/Source/Runtime/Engine/Private/WorldSettings.cpp
@@ -8,6 +8,7 @@
 #include "GameFramework/DefaultPhysicsVolume.h"
 #include "EngineUtils.h"
 #include "Engine/AssetUserData.h"
+#include "Engine/Engine.h"
 #include "Engine/WorldComposition.h"
 #include "WorldPartition/WorldPartition.h"
 #include "Net/UnrealNetwork.h"
diff --git a/Engine/Source/Runtime/GameplayMediaEncoder/Public/GameplayMediaEncoder.h b/Engine/Source/Runtime/GameplayMediaEncoder/Public/GameplayMediaEncoder.h
--- a/Engine/Source/Runtime/GameplayMediaEncoder/Public/GameplayMediaEncoder.h
+++ b/Engine/Source/Runtime/GameplayMediaEncoder/Public/GameplayMediaEncoder.h
@@ -3,6 +3,7 @@
 #pragma once
 
 #include "CoreMinimal.h"
+#include "ISubmixBufferListener.h"
 
 #include "Logging/LogMacros.h"
 #include "AudioMixerDevice.h"
diff --git a/Engine/Source/Runtime/MovieSceneCapture/Public/Protocols/AudioCaptureProtocol.h b/Engine/Source/Runtime/MovieSceneCapture/Public/Protocols/AudioCaptureProtocol.h
--- a/Engine/Source/Runtime/MovieSceneCapture/Public/Protocols/AudioCaptureProtocol.h
+++ b/Engine/Source/Runtime/MovieSceneCapture/Public/Protocols/AudioCaptureProtocol.h
@@ -2,6 +2,7 @@
 
 #pragma once
 #include "CoreMinimal.h"
+#include "Engine/Engine.h"
 #include "Engine/GameViewportClient.h"
 #include "Slate/SceneViewport.h"
 #include "AudioDevice.h"
diff --git a/Engine/Source/Runtime/MovieSceneTracks/Private/Evaluation/MovieSceneAudioTemplate.cpp b/Engine/Source/Runtime/MovieSceneTracks/Private/Evaluation/MovieSceneAudioTemplate.cpp
--- a/Engine/Source/Runtime/MovieSceneTracks/Private/Evaluation/MovieSceneAudioTemplate.cpp
+++ b/Engine/Source/Runtime/MovieSceneTracks/Private/Evaluation/MovieSceneAudioTemplate.cpp
@@ -3,6 +3,7 @@
 #include "Evaluation/MovieSceneAudioTemplate.h"
 
 
+#include "Engine/Engine.h"
 #include "Engine/EngineTypes.h"
 #include "Components/AudioComponent.h"
 #include "Sound/SoundCue.h"
diff --git a/Engine/Source/Runtime/SoundFieldRendering/Private/SoundFieldRendering.cpp b/Engine/Source/Runtime/SoundFieldRendering/Private/SoundFieldRendering.cpp
--- a/Engine/Source/Runtime/SoundFieldRendering/Private/SoundFieldRendering.cpp
+++ b/Engine/Source/Runtime/SoundFieldRendering/Private/SoundFieldRendering.cpp
@@ -1,6 +1,7 @@
 // Copyright Epic Games, Inc. All Rights Reserved.
 
 #include "SoundFieldRendering.h"
+#include "Engine/Engine.h"
 #include "HAL/IConsoleManager.h"
 #include "DSP/Dsp.h"
 #include "AudioMixer.h"
