Fixing errors in MediaIOCore after Renderer changes.

diff --git a/Engine/Plugins/Media/MediaIOFramework/Source/MediaIOCore/MediaIOCore.Build.cs b/Engine/Plugins/Media/MediaIOFramework/Source/MediaIOCore/MediaIOCore.Build.cs
--- a/Engine/Plugins/Media/MediaIOFramework/Source/MediaIOCore/MediaIOCore.Build.cs
+++ b/Engine/Plugins/Media/MediaIOFramework/Source/MediaIOCore/MediaIOCore.Build.cs
@@ -29,6 +29,13 @@ namespace UnrealBuildTool.Rules
 					"TimeManagement"
 				});
 
+			PrivateIncludePaths.AddRange(
+				new string[] {
+				// required for scene view extensions related headers
+				Path.Combine(GetModuleDirectory("Renderer"), "Private")
+			});
+
+
 			PrivateDependencyModuleNames.AddRange(
                 new string[]
                 {
diff --git a/Engine/Plugins/Media/MediaIOFramework/Source/MediaIOCore/Private/MediaCaptureSceneViewExtension.h b/Engine/Plugins/Media/MediaIOFramework/Source/MediaIOCore/Private/MediaCaptureSceneViewExtension.h
--- a/Engine/Plugins/Media/MediaIOFramework/Source/MediaIOCore/Private/MediaCaptureSceneViewExtension.h
+++ b/Engine/Plugins/Media/MediaIOFramework/Source/MediaIOCore/Private/MediaCaptureSceneViewExtension.h
@@ -7,13 +7,13 @@
 #include "MediaIOCoreModule.h"
 #include "ImagePixelData.h"
 #include "PostProcess/PostProcessing.h"
-#include "PostProcess/PostProcessMaterial.h"
 #include "RenderGraphUtils.h"
 #include "RHI.h"
 #include "RHICommandList.h"
 #include "RHIResources.h"
 #include "SceneView.h"
 #include "ScreenPass.h"
+#include "PostProcess/PostProcessMaterialInputs.h"
 
 class FRDGTexture;
 
