- Fixing network issues caused by rivermax reception thread
- Removing constraint on clock type for rivermax genlock
- Fixing non gpudirect rivermax output not working with experimental scheduling
- Changing default task priority for rivermax output with a cvar to change if needed

diff --git a/Engine/Plugins/Media/MediaIOFramework/Source/MediaIOCore/Private/MediaCapture.cpp b/Engine/Plugins/Media/MediaIOFramework/Source/MediaIOCore/Private/MediaCapture.cpp
--- a/Engine/Plugins/Media/MediaIOFramework/Source/MediaIOCore/Private/MediaCapture.cpp
+++ b/Engine/Plugins/Media/MediaIOFramework/Source/MediaIOCore/Private/MediaCapture.cpp
@@ -147,8 +147,8 @@ namespace UE::MediaCaptureData
 		/** Returns true if the readback is ready to be used */
 		virtual bool IsReadbackReady(FRHIGPUMask GPUMask) = 0;
 		
-		virtual void EnqueueCopy(FRDGBuilder& RDGBuilder, FRDGViewableResource* ResourceToReadback) const = 0;
-
+		virtual void EnqueueCopy(FRDGBuilder& RDGBuilder, FRDGViewableResource* ResourceToReadback) = 0;
+		
 		virtual FRHITexture* GetTextureResource() { return nullptr; }
 		virtual FRHIBuffer* GetBufferResource() { return nullptr; }
 		virtual FRHIResource* GetRHIResource() = 0;
@@ -227,7 +227,7 @@ namespace UE::MediaCaptureData
 
 		virtual void Unlock_Unsafe(FRHIGPUMask GPUMask) override
 		{
-			GDynamicRHI->RHIUnmapStagingSurface(ReadbackTexture->DestinationStagingTextures[0], GPUMask.GetFirstIndex());
+			GDynamicRHI->RHIUnmapStagingSurface(ReadbackTexture->DestinationStagingTextures[GPUMask.GetFirstIndex()], GPUMask.GetFirstIndex());
 		}
 
 		virtual bool IsReadbackReady(FRHIGPUMask GPUMask) override
@@ -244,7 +244,7 @@ namespace UE::MediaCaptureData
 		}
 
 		/** Adds a readback pass to the graph */
-		virtual void EnqueueCopy(FRDGBuilder& RDGBuilder, FRDGViewableResource* ResourceToReadback) const override
+		virtual void EnqueueCopy(FRDGBuilder& RDGBuilder, FRDGViewableResource* ResourceToReadback) override
 		{
 			AddEnqueueCopyPass(RDGBuilder, ReadbackTexture.Get(), static_cast<FRDGTexture*>(ResourceToReadback));
 		}
@@ -265,7 +265,7 @@ namespace UE::MediaCaptureData
 		TUniquePtr<FRHIGPUTextureReadback> ReadbackTexture;
 	};
 
-	class FBufferCaptureFrame : public FCaptureFrame
+	class FBufferCaptureFrame : public FCaptureFrame, public TSharedFromThis<UE::MediaCaptureData::FBufferCaptureFrame, ESPMode::ThreadSafe>
 	{
 	public:
 		/** Type alias for the output resource type used during capture frame */
@@ -305,10 +305,22 @@ namespace UE::MediaCaptureData
 			return ReadbackBuffer->Lock(Buffer->GetRHI()->GetSize());
 		}
 
+		virtual void* Lock_Unsafe(FRHIGPUMask GPUMask, int32& OutRowStride) override
+		{
+			void* ReadbackPointer = GDynamicRHI->RHILockStagingBuffer(DestinationStagingBuffers[GPUMask.GetFirstIndex()], nullptr, 0, Buffer->GetRHI()->GetSize());
+			OutRowStride = Buffer->GetRHI()->GetStride();
+			return ReadbackPointer;
+		}
+		
 		virtual void Unlock() override
 		{
 			ReadbackBuffer->Unlock();
 		}
+		
+		virtual void Unlock_Unsafe(FRHIGPUMask GPUMask) override
+		{
+			GDynamicRHI->RHIUnlockStagingBuffer(DestinationStagingBuffers[GPUMask.GetFirstIndex()]);
+		}
 
 		virtual bool IsReadbackReady(FRHIGPUMask GPUMask) override
 		{
@@ -322,10 +334,45 @@ namespace UE::MediaCaptureData
 			return RDGBuilder.RegisterExternalBuffer(Buffer, TEXT("OutputBuffer"));
 		}
 
+		BEGIN_SHADER_PARAMETER_STRUCT(FEnqueueCopyBufferPass, )
+		RDG_BUFFER_ACCESS(Buffer, ERHIAccess::CopySrc)
+		END_SHADER_PARAMETER_STRUCT()
+
+
 		/** Adds a readback pass to the graph */
-		virtual void EnqueueCopy(FRDGBuilder& RDGBuilder, FRDGViewableResource* ResourceToReadback) const
-		{
-			AddEnqueueCopyPass(RDGBuilder, ReadbackBuffer.Get(), static_cast<FRDGBuffer*>(ResourceToReadback), Buffer->GetRHI()->GetSize());
+		virtual void EnqueueCopy(FRDGBuilder& RDGBuilder, FRDGViewableResource* ResourceToReadback) override
+		{
+			if (CVarMediaIOScheduleOnAnyThread.GetValueOnAnyThread() && CVarMediaIOEnableExperimentalScheduling.GetValueOnAnyThread())
+            {
+				FEnqueueCopyBufferPass* PassParameters = RDGBuilder.AllocParameters<FEnqueueCopyBufferPass>();
+				PassParameters->Buffer = static_cast<FRDGBuffer*>(ResourceToReadback);
+
+				TSharedPtr<FBufferCaptureFrame> CaptureFramePtr = AsShared();
+				RDGBuilder.AddPass(
+					RDG_EVENT_NAME("EnqueueCopy(%s)", ResourceToReadback->Name),
+					PassParameters,
+					ERDGPassFlags::Readback,
+					[CaptureFramePtr, ResourceToReadback](FRHICommandList& RHICmdList)
+				{
+					CaptureFramePtr->LastCopyGPUMask = RHICmdList.GetGPUMask();
+
+					for (uint32 GPUIndex : CaptureFramePtr->LastCopyGPUMask)
+					{
+						SCOPED_GPU_MASK(RHICmdList, FRHIGPUMask::FromIndex(GPUIndex));
+
+						if (!CaptureFramePtr->DestinationStagingBuffers[GPUIndex])
+						{
+							CaptureFramePtr->DestinationStagingBuffers[GPUIndex] = RHICreateStagingBuffer();
+						}
+
+						RHICmdList.CopyToStagingBuffer(static_cast<FRHIBuffer*>(ResourceToReadback->GetRHI()), CaptureFramePtr->DestinationStagingBuffers[GPUIndex], 0, CaptureFramePtr->Buffer->GetRHI()->GetSize());
+					}
+				});
+            }
+			else
+			{
+				AddEnqueueCopyPass(RDGBuilder, ReadbackBuffer.Get(), static_cast<FRDGBuffer*>(ResourceToReadback), Buffer->GetRHI()->GetSize());
+			}
 		}
 
 		virtual FRHIBuffer* GetBufferResource() override
@@ -342,6 +389,17 @@ namespace UE::MediaCaptureData
 	public:
 		TRefCountPtr<FRDGPooledBuffer> Buffer;
 		TUniquePtr<FRHIGPUBufferReadback> ReadbackBuffer;
+		
+// Used for the ExperimentalScheduling path
+#if WITH_MGPU
+		FStagingBufferRHIRef DestinationStagingBuffers[MAX_NUM_GPUS];
+#else
+		FStagingBufferRHIRef DestinationStagingBuffers[1];
+#endif
+
+		// Copied from FRHIGPUReadback.h
+		FRHIGPUMask LastCopyGPUMask;
+		uint32 LastLockGPUIndex = 0; 
 	};
 
 	/** Helper struct to contain arguments for CaptureFrame */
@@ -600,7 +658,7 @@ namespace UE::MediaCaptureData
 
 						TRACE_CPUPROFILER_EVENT_SCOPE(UMediaCapture::RHIResourceCaptured);
 						TRACE_CPUPROFILER_EVENT_SCOPE_TEXT(*FString::Printf(TEXT("MediaCapture Output Frame %d"), ReadyFrame->CaptureBaseData.SourceFrameNumberRenderThread));
-						SCOPE_CYCLE_COUNTER(STAT_MediaCapture_RenderThread_CaptureCallback)
+						SCOPE_CYCLE_COUNTER(STAT_MediaCapture_RenderThread_RHI_CaptureCallback)
 
 						if (IsInRenderingThread())
 						{
@@ -633,13 +691,10 @@ namespace UE::MediaCaptureData
 					int32 RowStride = 0;
 
 					// Readback should be ready since we're after the sync point.
-					ensure(ReadyFrame->IsReadbackReady(GPUMask));
-					{
-						SCOPE_CYCLE_COUNTER(STAT_MediaCapture_RenderThread_LockResource);
-						ColorDataBuffer = ReadyFrame->Lock_Unsafe(GPUMask, RowStride);
-					}
+					SCOPE_CYCLE_COUNTER(STAT_MediaCapture_RenderThread_LockResource);
+					ColorDataBuffer = ReadyFrame->Lock_Unsafe(GPUMask, RowStride);
 
-					if (ColorDataBuffer)
+					if (ensure(ColorDataBuffer))
 					{
 						{
 							TRACE_CPUPROFILER_EVENT_SCOPE_TEXT(*FString::Printf(TEXT("MediaCapture Output Frame %d"), ReadyFrame->GetId()));
@@ -731,14 +786,15 @@ namespace UE::MediaCaptureData
 												break;
 											}
 
-											FPlatformProcess::SleepNoStats(0);
+											constexpr float SleepTimeSeconds = 50 * 1E-6;
+											FPlatformProcess::SleepNoStats(SleepTimeSeconds);
 										}
 
 										SyncDataPtr->RHIFence->Clear();
 										SyncDataPtr->bIsBusy = false;
 									}
 
-									if (CapturingFrame->bMediaCaptureActive && bWaitedForCompletion)
+									if (CapturingFrame->bMediaCaptureActive && bWaitedForCompletion && MediaCapture)
 									{
 										if (MediaCapture->UseAnyThreadCapture())
 										{
@@ -1106,6 +1162,7 @@ bool UMediaCapture::StartSourceCapture(TSharedPtr<UE::MediaCapture::Private::FCa
 	MediaCaptureAnalytics::SendCaptureEvent(GetCaptureSourceType());
 #endif
 
+
 	if (!GRHISupportsMultithreading)
 	{
 		UE_LOG(LogMediaIOCore, Display, TEXT("AnyThread Capture was disabled because the current RHI does not support Multithreading."));
diff --git a/Engine/Plugins/VirtualProduction/Rivermax/RivermaxCore/Source/RivermaxCore/Private/RivermaxInputStream.cpp b/Engine/Plugins/VirtualProduction/Rivermax/RivermaxCore/Source/RivermaxCore/Private/RivermaxInputStream.cpp
--- a/Engine/Plugins/VirtualProduction/Rivermax/RivermaxCore/Source/RivermaxCore/Private/RivermaxInputStream.cpp
+++ b/Engine/Plugins/VirtualProduction/Rivermax/RivermaxCore/Source/RivermaxCore/Private/RivermaxInputStream.cpp
@@ -306,6 +306,8 @@ namespace UE::RivermaxCore::Private
 		{
 			UE_LOG(LogRivermax, Warning, TEXT("Rivermax Input stream failed to get next chunk. Status: %d"), Status);
 		}
+
+		FPlatformProcess::SleepNoStats(UE::RivermaxCore::Private::Utils::SleepTimeSeconds);
 	}
 
 	bool FRivermaxInputStream::Init()
@@ -976,6 +978,7 @@ namespace UE::RivermaxCore::Private
 					const double CallbackTimestamp = FPlatformTime::Seconds();
 					while (CallbackPayload->bIsWaitingForPendingCopy == true && bIsShuttingDown == false)
 					{
+						FPlatformProcess::SleepNoStats(UE::RivermaxCore::Private::Utils::SleepTimeSeconds);
 						if (FPlatformTime::Seconds() - CallbackTimestamp > CVarWaitForCompletionTimeout.GetValueOnAnyThread())
 						{
 							UE_LOG(LogRivermax, Error, TEXT("Waiting for gpu copy of sample timed out."));
diff --git a/Engine/Plugins/VirtualProduction/Rivermax/RivermaxCore/Source/RivermaxCore/Private/RivermaxUtils.h b/Engine/Plugins/VirtualProduction/Rivermax/RivermaxCore/Source/RivermaxCore/Private/RivermaxUtils.h
--- a/Engine/Plugins/VirtualProduction/Rivermax/RivermaxCore/Source/RivermaxCore/Private/RivermaxUtils.h
+++ b/Engine/Plugins/VirtualProduction/Rivermax/RivermaxCore/Source/RivermaxCore/Private/RivermaxUtils.h
@@ -26,6 +26,9 @@ namespace UE::RivermaxCore::Private::Utils
 
 	/** SMTPE 2110-10.The Media Clock and RTP Clock rate for streams compliant to this standard shall be 90 kHz. */
 	static constexpr double MediaClockSampleRate = 90000.0;
+	
+	/** Common sleep time used in places where we are waiting for something to complete */
+	static constexpr float SleepTimeSeconds = 50 * 1E-6;
 
 	/** Convert a set of streaming option to its SDP description. Currently only support video type. */
 	void StreamOptionsToSDPDescription(const UE::RivermaxCore::FRivermaxOutputStreamOptions& Options, FAnsiStringBuilderBase& OutSDPDescription);
diff --git a/Engine/Plugins/VirtualProduction/Rivermax/RivermaxMedia/Source/RivermaxMedia/Private/RivermaxCustomTimeStep.cpp b/Engine/Plugins/VirtualProduction/Rivermax/RivermaxMedia/Source/RivermaxMedia/Private/RivermaxCustomTimeStep.cpp
--- a/Engine/Plugins/VirtualProduction/Rivermax/RivermaxMedia/Source/RivermaxMedia/Private/RivermaxCustomTimeStep.cpp
+++ b/Engine/Plugins/VirtualProduction/Rivermax/RivermaxMedia/Source/RivermaxMedia/Private/RivermaxCustomTimeStep.cpp
@@ -22,20 +22,16 @@ bool URivermaxCustomTimeStep::Initialize(UEngine* InEngine)
 		return false;
 	}
 
-	if (RivermaxModule->IsAvailable() == false)
+	if (RivermaxModule->GetRivermaxManager()->IsInitialized() == false)
 	{
+		UE_LOG(LogRivermaxMedia, Warning, TEXT("Can't initialize Rivermax custom timestep. Library isn't available."))
 		State = ECustomTimeStepSynchronizationState::Error;
 		return false;
 	}
 
 	if (RivermaxModule->GetRivermaxManager()->GetDevices().Num() <= 0)
 	{
-		State = ECustomTimeStepSynchronizationState::Error;
-		return false;
-	}
-
-	if (RivermaxModule->GetRivermaxManager()->GetTimeSource() != ERivermaxTimeSource::PTP)
-	{
+		UE_LOG(LogRivermaxMedia, Warning, TEXT("Can't initialize Rivermax custom timestep. No devices were found."))
 		State = ECustomTimeStepSynchronizationState::Error;
 		return false;
 	}
diff --git a/Engine/Plugins/VirtualProduction/Rivermax/RivermaxMedia/Source/RivermaxMedia/Private/RivermaxMediaCapture.cpp b/Engine/Plugins/VirtualProduction/Rivermax/RivermaxMedia/Source/RivermaxMedia/Private/RivermaxMediaCapture.cpp
--- a/Engine/Plugins/VirtualProduction/Rivermax/RivermaxMedia/Source/RivermaxMedia/Private/RivermaxMediaCapture.cpp
+++ b/Engine/Plugins/VirtualProduction/Rivermax/RivermaxMedia/Source/RivermaxMedia/Private/RivermaxMediaCapture.cpp
@@ -2,6 +2,7 @@
 
 #include "RivermaxMediaCapture.h"
 
+#include "Async/Fundamental/Task.h"
 #include "IRivermaxCoreModule.h"
 #include "IRivermaxManager.h"
 #include "IRivermaxOutputStream.h"
@@ -23,6 +24,13 @@ DECLARE_GPU_STAT(Rivermax_Capture);
 DECLARE_GPU_STAT(Rivermax_SyncPointPass);
 
 
+static TAutoConsoleVariable<int32> CVarRivermaxPollTaskPriority(
+	TEXT("Rivermax.Output.PollTaskPriority"), static_cast<int32>(LowLevelTasks::ETaskPriority::High),
+	TEXT("Priority of the task responsible to poll the render fence"),
+	ECVF_Default);
+
+
+
 /** Structure holding data used for synchronization of buffer output */
 struct URivermaxMediaCapture::FRivermaxCaptureSyncData
 {
@@ -355,7 +363,8 @@ void URivermaxMediaCapture::AddSyncPointPass(FRDGBuilder& GraphBuilder, const FC
 						// Wait until fence has been written (shader has completed)
 						while (SyncDataPtr->RHIFence->Poll() == false && Capturer->bIsActive)
 						{
-							FPlatformProcess::SleepNoStats(0);
+							constexpr float SleepTimeSeconds = 50 * 1E-6;
+							FPlatformProcess::SleepNoStats(SleepTimeSeconds);
 						}
 
 						SyncDataPtr->RHIFence->Clear();
@@ -370,7 +379,7 @@ void URivermaxMediaCapture::AddSyncPointPass(FRDGBuilder& GraphBuilder, const FC
 							}
 						}
 
-					}, LowLevelTasks::ETaskPriority::BackgroundHigh);
+					}, static_cast<LowLevelTasks::ETaskPriority>(CVarRivermaxPollTaskPriority.GetValueOnAnyThread()));
 				}
 			}
 		});
