MediaCapture: Fixed getting starved of buffers when certain race condition occurred.

OnReadbackComplete is called from an async task that waits on a GPU fence until the frame is ready, but if the task doesn't wake up promptly and another equivalent task processes a future frame first, the out of order situation will prevent the frames from being marked as available, effectively losing them from the pool until it gets starved of enough buffers for proper operation.

This CL adds a gate such that the async tasks processes the frames in order, both avoiding the buffer starvation and a potential race condition when accessing the queues that hold the available and pending frames.

Also changed nDisplay to flush when media capture overruns instead of skip.

To reduce the chances of the async task waking up late, the task priority was changed to high, and adjustable via MediaIO.Capture.PollTaskPriority.

diff --git a/Engine/Plugins/Media/MediaIOFramework/Source/MediaIOCore/Private/MediaCapture.cpp b/Engine/Plugins/Media/MediaIOFramework/Source/MediaIOCore/Private/MediaCapture.cpp
--- a/Engine/Plugins/Media/MediaIOFramework/Source/MediaIOCore/Private/MediaCapture.cpp
+++ b/Engine/Plugins/Media/MediaIOFramework/Source/MediaIOCore/Private/MediaCapture.cpp
@@ -53,6 +53,11 @@ static TAutoConsoleVariable<int32> CVarMediaIOScheduleOnAnyThread(
 	TEXT("Whether to wait for resource readback in a separate thread. (Experimental)"),
 	ECVF_RenderThreadSafe);
 
+static TAutoConsoleVariable<int32> CVarMediaIOCapturePollTaskPriority(
+	TEXT("MediaIO.Capture.PollTaskPriority"), static_cast<int32>(LowLevelTasks::ETaskPriority::High),
+	TEXT("Priority of the task responsible to poll the render fence"),
+	ECVF_RenderThreadSafe);
+
 /** Time spent in media capture sending a frame. */
 DECLARE_CYCLE_STAT(TEXT("MediaCapture RenderThread FrameCapture"), STAT_MediaCapture_RenderThread_FrameCapture, STATGROUP_Media);
 DECLARE_CYCLE_STAT(TEXT("MediaCapture RenderThread LockResource"), STAT_MediaCapture_RenderThread_LockResource, STATGROUP_Media);
@@ -761,11 +766,21 @@ namespace UE::MediaCaptureData
 							RHICmdList.WriteGPUFence(SyncDataPtr->RHIFence);
 							SyncDataPtr->bIsBusy = true;
 
+							// Here we request a number to process the frames in order.
+							const uint32 ExecutionNumber = MediaCapture->OrderedAsyncGateCaptureReady.GetANumber();
+
 							// Spawn a task that will wait (poll) and continue the process of providing a new texture
-							UE::Tasks::FTask Task = UE::Tasks::Launch(UE_SOURCE_LOCATION, [MediaCapture, CapturingFrame, SyncDataPtr]()
+							UE::Tasks::FTask Task = UE::Tasks::Launch(UE_SOURCE_LOCATION, [MediaCapture, CapturingFrame, SyncDataPtr, ExecutionNumber]()
 								{
 									TRACE_CPUPROFILER_EVENT_SCOPE(UMediaCapture::SyncPass);
 
+									ON_SCOPE_EXIT
+									{
+										// Make sure that we give up our turn to execute when exiting this task,
+										// which will allow the other async tasks to execute after waiting for the fence.
+										MediaCapture->OrderedAsyncGateCaptureReady.GiveUpTurn(ExecutionNumber);
+									};
+
 									double WaitTime = 0.0;
 									bool bWaitedForCompletion = false;
 									{
@@ -796,6 +811,20 @@ namespace UE::MediaCaptureData
 
 									if (CapturingFrame->bMediaCaptureActive && bWaitedForCompletion && MediaCapture)
 									{
+										// Ensure that we do not run the following code out of order with respect to the other sibling async tasks,
+										// because the Pending Frames are expected to be processed in order.
+										if (!MediaCapture->OrderedAsyncGateCaptureReady.IsMyTurn(ExecutionNumber))
+										{
+											TRACE_CPUPROFILER_EVENT_SCOPE(UMediaCapture::SyncPass::OutOfOrderWait);
+
+											UE_LOG(LogMediaIOCore, Warning, TEXT(
+												"The wait for the GPU fence for the next frame of MediaCapture '%s' unexpectedly happened out of order. "
+												"Order will be enforced by waiting until the previous frames are processed. "),
+												*MediaCapture->MediaOutputName);
+
+											MediaCapture->OrderedAsyncGateCaptureReady.WaitForTurn(ExecutionNumber);
+										}
+
 										if (MediaCapture->UseAnyThreadCapture())
 										{
 											OnReadbackComplete(FRHICommandListExecutor::GetImmediateCommandList(), MediaCapture, CapturingFrame);
@@ -811,7 +840,7 @@ namespace UE::MediaCaptureData
 										}
 									}
 
-								});
+								}, static_cast<LowLevelTasks::ETaskPriority>(CVarMediaIOCapturePollTaskPriority.GetValueOnRenderThread()));
 
 							{
 								// Add the task to the capture's list of pending tasks.
@@ -819,6 +848,12 @@ namespace UE::MediaCaptureData
 								MediaCapture->PendingReadbackTasks.Add(MoveTemp(Task));
 							}
 						}
+						else
+						{
+							UE_LOG(LogMediaIOCore, Error, TEXT(
+								"GetAvailableSyncHandler of MediaCapture '%s' failed to provide a fence, the captured buffers may not become available anymore."),
+								*MediaCapture->MediaOutputName);
+						}
 					}
 				});
 		}
@@ -1508,7 +1543,9 @@ void UMediaCapture::CaptureImmediate_RenderThread(const UE::MediaCaptureData::FC
 
 			if (!CapturingFrame)
 			{
-				UE_LOG(LogMediaIOCore, Error, TEXT("[%s] - No frames available for capture. This should not happen."), *MediaOutputName);
+				UE_LOG(LogMediaIOCore, Error, TEXT("[%s] - No frames available for capture for frame %llu. This should not happen."), 
+					*MediaOutputName, GFrameCounterRenderThread);
+
 				SetState(EMediaCaptureState::Error);
 				return;
 			}
@@ -1516,7 +1553,8 @@ void UMediaCapture::CaptureImmediate_RenderThread(const UE::MediaCaptureData::FC
 		else
 		{
 			//In case we are skipping frames, just keep capture frame as invalid
-			UE_LOG(LogMediaIOCore, Warning, TEXT("[%s] - No frames available for capture. Skipping"), *MediaOutputName);
+			UE_LOG(LogMediaIOCore, Warning, TEXT("[%s] - No frames available for capture of frame %llu. Skipping"), 
+				*MediaOutputName, GFrameCounterRenderThread);
 		}
 	}
 
@@ -1729,7 +1767,7 @@ void UMediaCapture::ProcessCapture_GameThread()
 		CapturingFrame = FrameManager->GetNextAvailable<FCaptureFrame>();
 	}
 
-	// Handle frame overrun (couldn't acquire a frame
+	// Handle frame overrun (couldn't acquire a frame)
 	if (!CapturingFrame && GetState() != EMediaCaptureState::StopRequested)
 	{
 		if (DesiredCaptureOptions.OverrunAction == EMediaCaptureOverrunAction::Flush)
diff --git a/Engine/Plugins/Media/MediaIOFramework/Source/MediaIOCore/Public/MediaCapture.h b/Engine/Plugins/Media/MediaIOFramework/Source/MediaIOCore/Public/MediaCapture.h
--- a/Engine/Plugins/Media/MediaIOFramework/Source/MediaIOCore/Public/MediaCapture.h
+++ b/Engine/Plugins/Media/MediaIOFramework/Source/MediaIOCore/Public/MediaCapture.h
@@ -11,6 +11,7 @@
 #include "HAL/CriticalSection.h"
 #include "MediaOutput.h"
 #include "Misc/Timecode.h"
+#include "OrderedAsyncGate.h"
 #include "PixelFormat.h"
 #include "RenderGraphResources.h"
 #include "RHI.h"
@@ -633,6 +634,9 @@ private:
 	TSharedPtr<UE::MediaCapture::Private::FCaptureSource> CaptureSource;
 
 	/** Holds a view extension that callbacks  */
-	TSharedPtr<class FMediaCaptureSceneViewExtension> ViewExtension; 
+	TSharedPtr<class FMediaCaptureSceneViewExtension> ViewExtension;
+
+	/** Helper to ensure that the async processing of frames happens in the expected order */
+	UE::MediaIO::FOrderedAsyncGate OrderedAsyncGateCaptureReady;
 };
  
\ No newline at end of file
diff --git a/Engine/Plugins/Media/MediaIOFramework/Source/MediaIOCore/Public/OrderedAsyncGate.h b/Engine/Plugins/Media/MediaIOFramework/Source/MediaIOCore/Public/OrderedAsyncGate.h
new file mode 100644
--- /dev/null
+++ b/Engine/Plugins/Media/MediaIOFramework/Source/MediaIOCore/Public/OrderedAsyncGate.h
@@ -0,0 +1,71 @@
+// Copyright Epic Games, Inc. All Rights Reserved.
+
+#pragma once
+
+#include "GenericPlatform/GenericPlatformProcess.h"
+
+#include <atomic>
+
+namespace UE::MediaIO
+{
+	/*
+	 * Class to ensure that async tasks execute any critical path in order.
+	 * It works like clients at a shop getting a number, waiting for their turn, and
+	 * when they are done they give up their turn so that the next number can be serviced.
+	 * Its functions can be called from any thread.
+	 */
+	class FOrderedAsyncGate
+	{
+	public:
+
+		/** Gets a number that represents the when the number holder will get their turn. */
+		uint64 GetANumber()
+		{
+			return NextNumber++;
+		}
+
+		/** Returns true if it is the turn of the given number */
+		bool IsMyTurn(const uint32 Number) const
+		{
+			return Number == CurrentNumber;
+		}
+
+		/**
+		 * This function will block until it is the turn of the given number to execute.
+		 * It can be called multiple times until the turn is given up by calling GiveUpTurn.
+		 * The number given must have been received from a call to GetANumber.
+		 */
+		void WaitForTurn(const uint32 Number) const
+		{
+			while (!IsMyTurn(Number))
+			{
+				constexpr float SpinWaitTimeSeconds = 50e-6;
+
+				FPlatformProcess::SleepNoStats(SpinWaitTimeSeconds);
+			}
+		}
+
+		/**
+		 * Call this function when you are ready to give up your turn.
+		 * It must be called with the number received from GetANumber, and it must
+		 * be called whether WaitForTurn was called or not. Not giving up the turn
+		 * will prevent any others from getting their turn.
+		 * Note: It will internally call WaitForTurn.
+		 */
+		void GiveUpTurn(const uint32 Number)
+		{
+			WaitForTurn(Number);
+
+			++CurrentNumber;
+		}
+
+	private:
+
+		/** The current number that has its turn without having to wait any longer */
+		std::atomic<uint32> CurrentNumber = 0;
+
+		/** The number that will be given out to the next client that calls GetANumber */
+		std::atomic<uint32> NextNumber = 0;
+	};
+
+}
diff --git a/Engine/Plugins/Runtime/nDisplay/Source/DisplayClusterMedia/Private/Capture/DisplayClusterMediaCaptureBase.cpp b/Engine/Plugins/Runtime/nDisplay/Source/DisplayClusterMedia/Private/Capture/DisplayClusterMediaCaptureBase.cpp
--- a/Engine/Plugins/Runtime/nDisplay/Source/DisplayClusterMedia/Private/Capture/DisplayClusterMediaCaptureBase.cpp
+++ b/Engine/Plugins/Runtime/nDisplay/Source/DisplayClusterMedia/Private/Capture/DisplayClusterMediaCaptureBase.cpp
@@ -116,7 +116,7 @@ bool FDisplayClusterMediaCaptureBase::StartMediaCapture()
 	FMediaCaptureOptions MediaCaptureOptions;
 	MediaCaptureOptions.NumberOfFramesToCapture = -1;
 	MediaCaptureOptions.bSkipFrameWhenRunningExpensiveTasks = false;
-	MediaCaptureOptions.OverrunAction = EMediaCaptureOverrunAction::Skip;
+	MediaCaptureOptions.OverrunAction = EMediaCaptureOverrunAction::Flush;
 
 	const bool bCaptureStarted = MediaCapture->CaptureRHITexture(Descriptor, MediaCaptureOptions);
 	if (bCaptureStarted)
