Fix media capture not working with capture phase BeforePostProcessing

diff --git a/Engine/Plugins/Media/MediaIOFramework/Source/MediaIOCore/Private/MediaCapture.cpp b/Engine/Plugins/Media/MediaIOFramework/Source/MediaIOCore/Private/MediaCapture.cpp
--- a/Engine/Plugins/Media/MediaIOFramework/Source/MediaIOCore/Private/MediaCapture.cpp
+++ b/Engine/Plugins/Media/MediaIOFramework/Source/MediaIOCore/Private/MediaCapture.cpp
@@ -415,6 +415,7 @@ namespace UE::MediaCaptureData
 		FTexture2DRHIRef ResourceToCapture;
 		FRDGTextureRef RDGResourceToCapture = nullptr;
 		FIntPoint DesiredSize = FIntPoint::ZeroValue;
+		FIntRect SourceViewRect{0,0,0,0};
 
 		EPixelFormat GetFormat() const
 		{
@@ -459,7 +460,7 @@ namespace UE::MediaCaptureData
 		FVector2D SizeU = FVector2D::ZeroVector;
 		FVector2D SizeV = FVector2D::ZeroVector;
 	};
-
+	
 	/** Helper class to be able to friend it and call methods on input media capture */
 	class FMediaCaptureHelper
 	{
@@ -483,7 +484,25 @@ namespace UE::MediaCaptureData
 
 			if (Args.MediaCapture->DesiredCaptureOptions.Crop == EMediaCaptureCroppingType::None)
 			{
-				if (Args.DesiredSize.X != Args.GetSizeX() || Args.DesiredSize.Y != Args.GetSizeY())
+				if (Args.SourceViewRect.Area() != 0)
+				{
+					if (Args.DesiredSize.X != Args.SourceViewRect.Width() || Args.DesiredSize.Y != Args.SourceViewRect.Height())
+					{
+						UE_LOG(LogMediaIOCore, Error, TEXT("The capture will stop for '%s'. The Source size doesn't match with the user requested size. Requested: %d,%d  Source: %d,%d")
+							, *Args.MediaCapture->MediaOutputName
+							, Args.SourceViewRect.Width(), Args.SourceViewRect.Height()
+							, Args.GetSizeX(), Args.GetSizeY());
+
+						return false;
+					}
+					else
+					{
+						// If source view rect is passed, it will override the crop passed as argument.
+						Args.MediaCapture->DesiredCaptureOptions.Crop = EMediaCaptureCroppingType::Custom;
+						Args.MediaCapture->DesiredCaptureOptions.CustomCapturePoint = Args.SourceViewRect.Min;
+					}
+				}
+				else if (Args.SourceViewRect.Area() == 0 && (Args.DesiredSize.X != Args.GetSizeX() || Args.DesiredSize.Y != Args.GetSizeY()))
 				{
 					UE_LOG(LogMediaIOCore, Error, TEXT("The capture will stop for '%s'. The Source size doesn't match with the user requested size. Requested: %d,%d  Source: %d,%d")
 						, *Args.MediaCapture->MediaOutputName
@@ -893,7 +912,7 @@ namespace UE::MediaCaptureData
 					Args.MediaCapture->LockDMATexture_RenderThread(CapturingFrame->GetTextureResource());
 				}
 			}
-
+			
 			{
 				TRACE_CPUPROFILER_EVENT_SCOPE(UMediaCapture::FrameCapture);
 				SCOPE_CYCLE_COUNTER(STAT_MediaCapture_RenderThread_FrameCapture);
@@ -913,7 +932,7 @@ namespace UE::MediaCaptureData
 
 				// Register output resource used by the current capture method (texture or buffer)
 				FRDGViewableResource* OutputResource = CapturingFrame->RegisterResource(Args.GraphBuilder);
-				
+
 				{
 					TRACE_CPUPROFILER_EVENT_SCOPE(UMediaCapture::GraphSetup);
 					SCOPED_DRAW_EVENTF(Args.GraphBuilder.RHICmdList, MediaCapture, TEXT("MediaCapture"));
@@ -1442,22 +1461,26 @@ bool UMediaCapture::SetCaptureAudioDevice(const FAudioDeviceHandle& InAudioDevic
 	return bSuccess;
 }
 
-void UMediaCapture::CaptureImmediate_RenderThread(FRDGBuilder& GraphBuilder, FRHITexture* InSourceTexture)
+void UMediaCapture::CaptureImmediate_RenderThread(FRDGBuilder& GraphBuilder, FRHITexture* InSourceTexture, FIntRect SourceViewRect)
 {
 	UE::MediaCaptureData::FCaptureFrameArgs CaptureArgs{GraphBuilder};
 	CaptureArgs.MediaCapture = this;
     CaptureArgs.DesiredSize = DesiredSize;
     CaptureArgs.ResourceToCapture = InSourceTexture;
+	CaptureArgs.SourceViewRect = SourceViewRect;
+	
 	CaptureImmediate_RenderThread(CaptureArgs);
 }
 
 
-void UMediaCapture::CaptureImmediate_RenderThread(FRDGBuilder& GraphBuilder,  FRDGTextureRef InSourceTextureRef)
+void UMediaCapture::CaptureImmediate_RenderThread(FRDGBuilder& GraphBuilder,  FRDGTextureRef InSourceTextureRef, FIntRect SourceViewRect)
 {
 	UE::MediaCaptureData::FCaptureFrameArgs CaptureArgs{GraphBuilder};
 	CaptureArgs.MediaCapture = this;
 	CaptureArgs.DesiredSize = DesiredSize;
 	CaptureArgs.RDGResourceToCapture = InSourceTextureRef;
+	CaptureArgs.SourceViewRect = SourceViewRect;
+	
 	CaptureImmediate_RenderThread(CaptureArgs);
 }
 
@@ -1504,15 +1527,26 @@ void UMediaCapture::CaptureImmediate_RenderThread(const UE::MediaCaptureData::FC
 	// Get cached capture data from game thread. We want to find a cached frame matching current render thread frame number
 	bool bFoundMatchingData = false;
 	FQueuedCaptureData NextCaptureData;
-	while (CaptureDataQueue.Dequeue(NextCaptureData))
+
 	{
-		if (NextCaptureData.BaseData.SourceFrameNumberRenderThread == GFrameCounterRenderThread)
+		FScopeLock ScopeLock(&CaptureDataQueueCriticalSection);
+		for (auto It = CaptureDataQueue.CreateIterator(); It; ++It)
 		{
-			bFoundMatchingData = true;
-			break;
+			if (It->BaseData.SourceFrameNumberRenderThread == GFrameCounterRenderThread)
+			{
+				NextCaptureData = *It;
+				bFoundMatchingData = true;
+				It.RemoveCurrent();
+				break;
+			}
+			else if (GFrameCounterRenderThread > It->BaseData.SourceFrameNumberRenderThread &&
+					GFrameCounterRenderThread - It->BaseData.SourceFrameNumberRenderThread > MaxCaptureDataAgeInFrames)
+			{
+				// Remove old frame data that wasn't used.
+				It.RemoveCurrent();
+			}
 		}
 	}
-
 	if (bFoundMatchingData == false)
 	{
 		UE_LOG(LogMediaIOCore, Warning, TEXT("Can't capture frame. Could not find the matching game frame %d."), GFrameCounterRenderThread);
@@ -1922,7 +1956,7 @@ void UMediaCapture::InitializeOutputResources(int32 InNumberOfBuffers)
 
 
 					NewFrame->RenderTarget = AllocatePooledTexture(OutputDesc, *FString::Format(TEXT("MediaCapture RenderTarget {0}"), { Index }));
-
+					
 					// Only create CPU readback resource when we are using the CPU callback
 					if (!This->bShouldCaptureRHIResource)
 					{
@@ -1977,7 +2011,7 @@ void UMediaCapture::InitializeCaptureFrame(const TSharedPtr<UE::MediaCaptureData
 	{
 		CaptureFrame->CaptureBaseData.SourceFrameTimecode = FApp::GetTimecode();
 		CaptureFrame->CaptureBaseData.SourceFrameTimecodeFramerate = FApp::GetTimecodeFrameRate();
-		CaptureFrame->CaptureBaseData.SourceFrameNumberRenderThread = GFrameNumber;
+		CaptureFrame->CaptureBaseData.SourceFrameNumberRenderThread = GFrameCounter;
 		CaptureFrame->CaptureBaseData.SourceFrameNumber = ++CaptureRequestCount;
 		CaptureFrame->UserData = GetCaptureFrameUserData_GameThread();
 	}
@@ -2060,8 +2094,9 @@ void UMediaCapture::OnBeginFrame_GameThread()
 		CaptureData.BaseData.SourceFrameTimecodeFramerate = FApp::GetTimecodeFrameRate();
 		CaptureData.BaseData.SourceFrameNumberRenderThread = GFrameCounter;
 		CaptureData.UserData = GetCaptureFrameUserData_GameThread();
-		
-		CaptureDataQueue.Enqueue(MoveTemp(CaptureData));
+
+		FScopeLock Lock(&CaptureDataQueueCriticalSection);
+		CaptureDataQueue.Insert(MoveTemp(CaptureData), 0);
 	}
 }
 
@@ -2069,7 +2104,7 @@ void UMediaCapture::OnEndFrame_GameThread()
 {
 	using namespace UE::MediaCaptureData;
 
-	TRACE_CPUPROFILER_EVENT_SCOPE_TEXT(*FString::Printf(TEXT("MediaCapture End Frame %d"), GFrameNumber));
+	TRACE_CPUPROFILER_EVENT_SCOPE_TEXT(*FString::Printf(TEXT("MediaCapture End Frame %d"), GFrameCounter));
 
 	if (!bOutputResourcesInitialized)
 	{
diff --git a/Engine/Plugins/Media/MediaIOFramework/Source/MediaIOCore/Private/MediaCaptureSceneViewExtension.h b/Engine/Plugins/Media/MediaIOFramework/Source/MediaIOCore/Private/MediaCaptureSceneViewExtension.h
--- a/Engine/Plugins/Media/MediaIOFramework/Source/MediaIOCore/Private/MediaCaptureSceneViewExtension.h
+++ b/Engine/Plugins/Media/MediaIOFramework/Source/MediaIOCore/Private/MediaCaptureSceneViewExtension.h
@@ -65,34 +65,14 @@ public:
 		}
 	};
 	virtual void BeginRenderViewFamily(FSceneViewFamily& InViewFamily) override {}; 
-
-	virtual void PrePostProcessPass_RenderThread(FRDGBuilder& GraphBuilder, const FSceneView& View, const FPostProcessingInputs& Inputs) override
-	{
-		if (CapturePhase == EMediaCapturePhase::BeforePostProcessing)
-		{
-			Inputs.Validate();
-
-			if (FRDGTextureRef TextureRef = Inputs.ViewFamilyTexture)
-			{
-				RDG_EVENT_SCOPE(GraphBuilder, "MediaCaptureSceneExtension");
-				if (UMediaCapture* MediaCapture = WeakCapture.Get())
-				{
-					MediaCapture->CaptureImmediate_RenderThread(GraphBuilder, TextureRef);
-				}
-			}
-		}
-	}
-
+	
 	virtual void SubscribeToPostProcessingPass(EPostProcessingPass PassId, FAfterPassCallbackDelegateArray& InOutPassCallbacks, bool bIsPassEnabled) override
-	{		
+	{
 		if ((CapturePhase == EMediaCapturePhase::AfterMotionBlur && PassId == EPostProcessingPass::MotionBlur)
 			|| (CapturePhase == EMediaCapturePhase::AfterToneMap && PassId == EPostProcessingPass::Tonemap)
-			|| (CapturePhase == EMediaCapturePhase::AfterFXAA && PassId == EPostProcessingPass::FXAA))
+			|| (CapturePhase == EMediaCapturePhase::AfterFXAA && PassId == EPostProcessingPass::FXAA)
+			|| (CapturePhase == EMediaCapturePhase::BeforePostProcessing && PassId == EPostProcessingPass::SSRInput))
 		{
-			if (!bIsPassEnabled)
-			{
-				UE_LOG(LogMediaIOCore, Warning, TEXT("Media Capture will not work since the selected capture phase is not enabled."));
-			}
 			InOutPassCallbacks.Add(FAfterPassCallbackDelegate::CreateSP(this, &FMediaCaptureSceneViewExtension::PostProcessCallback_RenderThread));
 		}
 	}
@@ -108,13 +88,13 @@ public:
 
 		FScreenPassTexture SceneColor = InOutInputs.GetInput(EPostProcessMaterialInput::SceneColor);
 		InOutInputs.Validate();
-		
+
 		if (FRDGTextureRef TextureRef = SceneColor.Texture)
 		{
 			RDG_EVENT_SCOPE(GraphBuilder, "MediaCaptureSceneExtension");
 			if (WeakCapture.IsValid())
 			{
-				WeakCapture->CaptureImmediate_RenderThread(GraphBuilder, TextureRef);
+				WeakCapture->CaptureImmediate_RenderThread(GraphBuilder, TextureRef, SceneColor.ViewRect);
 			}
 		}
 
diff --git a/Engine/Plugins/Media/MediaIOFramework/Source/MediaIOCore/Public/MediaCapture.h b/Engine/Plugins/Media/MediaIOFramework/Source/MediaIOCore/Public/MediaCapture.h
--- a/Engine/Plugins/Media/MediaIOFramework/Source/MediaIOCore/Public/MediaCapture.h
+++ b/Engine/Plugins/Media/MediaIOFramework/Source/MediaIOCore/Public/MediaCapture.h
@@ -125,10 +125,10 @@ struct MEDIAIOCORE_API FRHICaptureResourceDescription
 UENUM(BlueprintType)
 enum class EMediaCapturePhase : uint8
 {
-	BeforePostProcessing,
+	BeforePostProcessing, // Will happen after TSR in order to get a texture of the right size.
+	AfterMotionBlur,
 	AfterToneMap,
 	AfterFXAA,
-	AfterMotionBlur,
 	EndFrame
 };
 
@@ -286,9 +286,9 @@ public:
 	bool UpdateTextureRenderTarget2D(UTextureRenderTarget2D* RenderTarget);
 	
 	/** Captures a resource immediately from the render thread. Used in RHI_RESOURCE capture mode */
-	void CaptureImmediate_RenderThread(FRDGBuilder& GraphBuilder, FRHITexture* InSourceTexture);
+	void CaptureImmediate_RenderThread(FRDGBuilder& GraphBuilder, FRHITexture* InSourceTexture, FIntRect SourceViewRect = FIntRect(0,0,0,0));
 	
-	void CaptureImmediate_RenderThread(FRDGBuilder& GraphBuilder, FRDGTextureRef InSourceTextureRef);
+	void CaptureImmediate_RenderThread(FRDGBuilder& GraphBuilder, FRDGTextureRef InSourceTextureRef, FIntRect SourceViewRect = FIntRect(0,0,0,0));
 
 	/**
 	 * Stop the previous requested capture.
@@ -607,7 +607,11 @@ private:
 		FCaptureBaseData BaseData;
 		TSharedPtr<FMediaCaptureUserData> UserData;
 	};
-	TSpscQueue<FQueuedCaptureData> CaptureDataQueue;
+
+	/** Used to synchronize access to the capture data queue. */
+	FCriticalSection CaptureDataQueueCriticalSection;
+	TArray<FQueuedCaptureData> CaptureDataQueue;
+	static constexpr int32 MaxCaptureDataAgeInFrames = 20;
 
 	/** Structure holding data used for synchronization of buffer output */
 	struct FMediaCaptureSyncData
