- Fix for mgpu crash in media capture

diff --git a/Engine/Plugins/Media/MediaIOFramework/Source/MediaIOCore/Private/MediaCapture.cpp b/Engine/Plugins/Media/MediaIOFramework/Source/MediaIOCore/Private/MediaCapture.cpp
--- a/Engine/Plugins/Media/MediaIOFramework/Source/MediaIOCore/Private/MediaCapture.cpp
+++ b/Engine/Plugins/Media/MediaIOFramework/Source/MediaIOCore/Private/MediaCapture.cpp
@@ -140,19 +140,19 @@ namespace UE::MediaCaptureData
 		virtual bool IsBufferResource() const = 0;
 
 		/** Locks the readback resource and returns a pointer to access data from system memory */
-		virtual void* Lock(FRHICommandListImmediate& RHICmdList, FRHIGPUMask GPUMask, int32& OutRowStride) = 0;
+		virtual void* Lock(FRHICommandListImmediate& RHICmdList, int32& OutRowStride) = 0;
 
-		virtual void* Lock_Unsafe(FRHIGPUMask GPUMask, int32& OutRowStride) { return nullptr; }
+		virtual void* Lock_Unsafe(int32& OutRowStride) { return nullptr; }
 
 		/** Unlocks the readback resource */
 		virtual void Unlock() = 0;
 
-		virtual void Unlock_Unsafe(FRHIGPUMask GPUMask) {}
+		virtual void Unlock_Unsafe() {}
 
 		/** Returns true if the readback is ready to be used */
 		virtual bool IsReadbackReady(FRHIGPUMask GPUMask) = 0;
 		
-		virtual void EnqueueCopy(FRDGBuilder& RDGBuilder, FRDGViewableResource* ResourceToReadback) = 0;
+		virtual void EnqueueCopy(FRDGBuilder& RDGBuilder, FRDGViewableResource* ResourceToReadback, bool bIsAnyThreadSupported) = 0;
 		
 		virtual FRHITexture* GetTextureResource() { return nullptr; }
 		virtual FRHIBuffer* GetBufferResource() { return nullptr; }
@@ -203,9 +203,9 @@ namespace UE::MediaCaptureData
 			return false;
 		}
 
-		virtual void* Lock(FRHICommandListImmediate& RHICmdList, FRHIGPUMask GPUMask, int32& OutRowStride) override
+		virtual void* Lock(FRHICommandListImmediate& RHICmdList, int32& OutRowStride) override
 		{
-			if (ReadbackTexture->IsReady(GPUMask) == false)
+			if (ReadbackTexture->IsReady() == false)
 			{
 				UE_LOG(LogMediaIOCore, Verbose, TEXT("Fence for texture readback was not ready"));
 			}
@@ -216,11 +216,11 @@ namespace UE::MediaCaptureData
 			return ReadbackPointer;
 		}
 
-		virtual void* Lock_Unsafe(FRHIGPUMask GPUMask, int32& OutRowStride) override
+		virtual void* Lock_Unsafe(int32& OutRowStride) override
 		{
 			void* ReadbackPointer = nullptr;
 			int32 ReadbackWidth, ReadbackHeight;
-			GDynamicRHI->RHIMapStagingSurface(ReadbackTexture->DestinationStagingTextures[GPUMask.GetFirstIndex()], nullptr, ReadbackPointer, ReadbackWidth, ReadbackHeight, GPUMask.GetFirstIndex());
+			GDynamicRHI->RHIMapStagingSurface(ReadbackTexture->DestinationStagingTextures[ReadbackTexture->GetLastCopyGPUMask().GetFirstIndex()], nullptr, ReadbackPointer, ReadbackWidth, ReadbackHeight, ReadbackTexture->GetLastCopyGPUMask().GetFirstIndex());
 			OutRowStride = ReadbackWidth * MediaCaptureDetails::GetBytesPerPixel(RenderTarget->GetDesc().Format);
 			return ReadbackPointer;
 		}
@@ -230,9 +230,9 @@ namespace UE::MediaCaptureData
 			ReadbackTexture->Unlock();
 		}
 
-		virtual void Unlock_Unsafe(FRHIGPUMask GPUMask) override
+		virtual void Unlock_Unsafe() override
 		{
-			GDynamicRHI->RHIUnmapStagingSurface(ReadbackTexture->DestinationStagingTextures[GPUMask.GetFirstIndex()], GPUMask.GetFirstIndex());
+			GDynamicRHI->RHIUnmapStagingSurface(ReadbackTexture->DestinationStagingTextures[ReadbackTexture->GetLastCopyGPUMask().GetFirstIndex()], ReadbackTexture->GetLastCopyGPUMask().GetFirstIndex());
 		}
 
 		virtual bool IsReadbackReady(FRHIGPUMask GPUMask) override
@@ -249,7 +249,7 @@ namespace UE::MediaCaptureData
 		}
 
 		/** Adds a readback pass to the graph */
-		virtual void EnqueueCopy(FRDGBuilder& RDGBuilder, FRDGViewableResource* ResourceToReadback) override
+		virtual void EnqueueCopy(FRDGBuilder& RDGBuilder, FRDGViewableResource* ResourceToReadback, bool bIsAnyThreadSupported) override
 		{
 			AddEnqueueCopyPass(RDGBuilder, ReadbackTexture.Get(), static_cast<FRDGTexture*>(ResourceToReadback));
 		}
@@ -298,9 +298,9 @@ namespace UE::MediaCaptureData
 			return true;
 		}
 
-		virtual void* Lock(FRHICommandListImmediate& RHICmdList, FRHIGPUMask GPUMask, int32& OutRowStride) override
+		virtual void* Lock(FRHICommandListImmediate& RHICmdList, int32& OutRowStride) override
 		{
-			if (ReadbackBuffer->IsReady(GPUMask) == false)
+			if (ReadbackBuffer->IsReady() == false)
 			{
 				UE_LOG(LogMediaIOCore, Verbose, TEXT("Fence for buffer readback was not ready, blocking."));
 				RHICmdList.BlockUntilGPUIdle();
@@ -310,9 +310,9 @@ namespace UE::MediaCaptureData
 			return ReadbackBuffer->Lock(Buffer->GetRHI()->GetSize());
 		}
 
-		virtual void* Lock_Unsafe(FRHIGPUMask GPUMask, int32& OutRowStride) override
+		virtual void* Lock_Unsafe(int32& OutRowStride) override
 		{
-			void* ReadbackPointer = GDynamicRHI->RHILockStagingBuffer(DestinationStagingBuffers[GPUMask.GetFirstIndex()], nullptr, 0, Buffer->GetRHI()->GetSize());
+			void* ReadbackPointer = GDynamicRHI->RHILockStagingBuffer(DestinationStagingBuffers[LastCopyGPUMask.GetFirstIndex()], nullptr, 0, Buffer->GetRHI()->GetSize());
 			OutRowStride = Buffer->GetRHI()->GetStride();
 			return ReadbackPointer;
 		}
@@ -322,9 +322,9 @@ namespace UE::MediaCaptureData
 			ReadbackBuffer->Unlock();
 		}
 		
-		virtual void Unlock_Unsafe(FRHIGPUMask GPUMask) override
+		virtual void Unlock_Unsafe() override
 		{
-			GDynamicRHI->RHIUnlockStagingBuffer(DestinationStagingBuffers[GPUMask.GetFirstIndex()]);
+			GDynamicRHI->RHIUnlockStagingBuffer(DestinationStagingBuffers[LastCopyGPUMask.GetFirstIndex()]);
 		}
 
 		virtual bool IsReadbackReady(FRHIGPUMask GPUMask) override
@@ -345,9 +345,9 @@ namespace UE::MediaCaptureData
 
 
 		/** Adds a readback pass to the graph */
-		virtual void EnqueueCopy(FRDGBuilder& RDGBuilder, FRDGViewableResource* ResourceToReadback) override
+		virtual void EnqueueCopy(FRDGBuilder& RDGBuilder, FRDGViewableResource* ResourceToReadback, bool bIsAnyThreadSupported) override
 		{
-			if (CVarMediaIOScheduleOnAnyThread.GetValueOnAnyThread() && CVarMediaIOEnableExperimentalScheduling.GetValueOnAnyThread())
+			if (bIsAnyThreadSupported)
             {
 				FEnqueueCopyBufferPass* PassParameters = RDGBuilder.AllocParameters<FEnqueueCopyBufferPass>();
 				PassParameters->Buffer = static_cast<FRDGBuffer*>(ResourceToReadback);
@@ -657,79 +657,109 @@ namespace UE::MediaCaptureData
 					}
 				};
 
-				FRHIGPUMask GPUMask;
-	#if WITH_MGPU
-				GPUMask = RHICmdList.GetGPUMask();
 
-				// If GPUMask is not set to a specific GPU we and since we are reading back the texture, it shouldn't matter which GPU we do this on.
-				if (!GPUMask.HasSingleIndex())
+
+				// Path where resource ready callback (readback / rhi capture) is on render thread (old method)
+				if (IsInRenderingThread())
 				{
-					GPUMask = FRHIGPUMask::FromIndex(GPUMask.GetFirstIndex());
-				}
+					// Scoped gpu mask shouldn't be needed for readback since we specify the gpu mask used during copy when we lock
+					// Keeping it for old render thread path
+					FRHIGPUMask GPUMask;
+#if WITH_MGPU
+					GPUMask = RHICmdList.GetGPUMask();
 
-				SCOPED_GPU_MASK(RHICmdList, GPUMask);
-	#endif
+					// If GPUMask is not set to a specific GPU we and since we are reading back the texture, it shouldn't matter which GPU we do this on.
+					if (!GPUMask.HasSingleIndex())
+					{
+						GPUMask = FRHIGPUMask::FromIndex(GPUMask.GetFirstIndex());
+					}
+					SCOPED_GPU_MASK(RHICmdList, GPUMask);
+#endif
 
-				// Are we doing a GPU Direct transfer
-				if (MediaCapture->ShouldCaptureRHIResource())
-				{
-					if (ReadyFrame->IsTextureResource())
+					// Are we doing a GPU Direct transfer
+					if (MediaCapture->ShouldCaptureRHIResource())
 					{
+						if (ReadyFrame->IsTextureResource())
 						{
-							TRACE_CPUPROFILER_EVENT_SCOPE(UMediaCapture::UnlockDMATexture_RenderThread);
-							MediaCapture->UnlockDMATexture_RenderThread(ReadyFrame->GetTextureResource());
-						}
+							{
+								TRACE_CPUPROFILER_EVENT_SCOPE(UMediaCapture::UnlockDMATexture_RenderThread);
+								MediaCapture->UnlockDMATexture_RenderThread(ReadyFrame->GetTextureResource());
+							}
 
-						TRACE_CPUPROFILER_EVENT_SCOPE(UMediaCapture::RHIResourceCaptured);
-						TRACE_CPUPROFILER_EVENT_SCOPE_TEXT(*FString::Printf(TEXT("MediaCapture Output Frame %d"), ReadyFrame->CaptureBaseData.SourceFrameNumberRenderThread));
-						SCOPE_CYCLE_COUNTER(STAT_MediaCapture_RenderThread_RHI_CaptureCallback)
+							TRACE_CPUPROFILER_EVENT_SCOPE(UMediaCapture::RHIResourceCaptured);
+							TRACE_CPUPROFILER_EVENT_SCOPE_TEXT(*FString::Printf(TEXT("MediaCapture Output Frame %d"), ReadyFrame->CaptureBaseData.SourceFrameNumberRenderThread));
+							SCOPE_CYCLE_COUNTER(STAT_MediaCapture_RenderThread_RHI_CaptureCallback)
 
-						if (IsInRenderingThread())
-						{
 							MediaCapture->OnRHIResourceCaptured_RenderingThread(ReadyFrame->CaptureBaseData, ReadyFrame->UserData, ReadyFrame->GetTextureResource());
 						}
 						else
 						{
-							MediaCapture->OnRHIResourceCaptured_AnyThread(ReadyFrame->CaptureBaseData, ReadyFrame->UserData, ReadyFrame->GetTextureResource());
+							TRACE_CPUPROFILER_EVENT_SCOPE(UMediaCapture::RHIResourceCaptured);
+							MediaCapture->OnRHIResourceCaptured_RenderingThread(ReadyFrame->CaptureBaseData, ReadyFrame->UserData, ReadyFrame->GetBufferResource());
 						}
 					}
 					else
 					{
-						TRACE_CPUPROFILER_EVENT_SCOPE(UMediaCapture::RHIResourceCaptured);
-						if (IsInRenderingThread())
-						{
-							MediaCapture->OnRHIResourceCaptured_RenderingThread(ReadyFrame->CaptureBaseData, ReadyFrame->UserData, ReadyFrame->GetBufferResource());
-						}
-						else
+						// Lock & read
+						void* ColorDataBuffer = nullptr;
+						int32 RowStride = 0;
+
+						// Readback should be ready since we're after the sync point.
+						SCOPE_CYCLE_COUNTER(STAT_MediaCapture_RenderThread_LockResource);
+						ColorDataBuffer = ReadyFrame->Lock(FRHICommandListExecutor::GetImmediateCommandList(), RowStride);
+
+						if (ensure(ColorDataBuffer))
 						{
-							MediaCapture->OnRHIResourceCaptured_AnyThread(ReadyFrame->CaptureBaseData, ReadyFrame->UserData, ReadyFrame->GetBufferResource());
+							{
+								TRACE_CPUPROFILER_EVENT_SCOPE_TEXT(*FString::Printf(TEXT("MediaCapture Output Frame %d"), ReadyFrame->GetId()));
+								SCOPE_CYCLE_COUNTER(STAT_MediaCapture_AnyThread_LockResource)
+								MediaCapture->OnFrameCaptured_RenderingThread(ReadyFrame->CaptureBaseData, ReadyFrame->UserData, ColorDataBuffer, MediaCapture->DesiredOutputSize.X, MediaCapture->DesiredOutputSize.Y, RowStride);
+							}
+
+							ReadyFrame->Unlock();
 						}
 					}
-
-					ReadyFrame->bDoingGPUCopy = false;
 				}
 				else
 				{
-					// Lock & read
-					void* ColorDataBuffer = nullptr;
-					int32 RowStride = 0;
-
-					// Readback should be ready since we're after the sync point.
-					SCOPE_CYCLE_COUNTER(STAT_MediaCapture_RenderThread_LockResource);
-					ColorDataBuffer = ReadyFrame->Lock_Unsafe(GPUMask, RowStride);
-
-					if (ensure(ColorDataBuffer))
+					// Are we doing a GPU Direct transfer
+					if (MediaCapture->ShouldCaptureRHIResource())
 					{
+						if (ReadyFrame->IsTextureResource())
 						{
-							TRACE_CPUPROFILER_EVENT_SCOPE_TEXT(*FString::Printf(TEXT("MediaCapture Output Frame %d"), ReadyFrame->GetId()));
-							SCOPE_CYCLE_COUNTER(STAT_MediaCapture_AnyThread_LockResource)
-
-							if (IsInRenderingThread())
 							{
-								MediaCapture->OnFrameCaptured_RenderingThread(ReadyFrame->CaptureBaseData, ReadyFrame->UserData, ColorDataBuffer, MediaCapture->DesiredOutputSize.X, MediaCapture->DesiredOutputSize.Y, RowStride);
+								TRACE_CPUPROFILER_EVENT_SCOPE(UMediaCapture::UnlockDMATexture_RenderThread);
+								MediaCapture->UnlockDMATexture_RenderThread(ReadyFrame->GetTextureResource());
 							}
-							else
+
+							TRACE_CPUPROFILER_EVENT_SCOPE(UMediaCapture::RHIResourceCaptured);
+							TRACE_CPUPROFILER_EVENT_SCOPE_TEXT(*FString::Printf(TEXT("MediaCapture Output Frame %d"), ReadyFrame->CaptureBaseData.SourceFrameNumberRenderThread));
+							SCOPE_CYCLE_COUNTER(STAT_MediaCapture_RenderThread_RHI_CaptureCallback)
+
+							MediaCapture->OnRHIResourceCaptured_AnyThread(ReadyFrame->CaptureBaseData, ReadyFrame->UserData, ReadyFrame->GetTextureResource());
+						}
+						else
+						{
+							TRACE_CPUPROFILER_EVENT_SCOPE(UMediaCapture::RHIResourceCaptured);
+							MediaCapture->OnRHIResourceCaptured_AnyThread(ReadyFrame->CaptureBaseData, ReadyFrame->UserData, ReadyFrame->GetBufferResource());
+						}
+					}
+					else
+					{
+						// Lock & read
+						void* ColorDataBuffer = nullptr;
+						int32 RowStride = 0;
+
+						// Readback should be ready since we're after the sync point.
+						SCOPE_CYCLE_COUNTER(STAT_MediaCapture_RenderThread_LockResource);
+						ColorDataBuffer = ReadyFrame->Lock_Unsafe(RowStride);
+
+						if (ensure(ColorDataBuffer))
+						{
 							{
+								TRACE_CPUPROFILER_EVENT_SCOPE_TEXT(*FString::Printf(TEXT("MediaCapture Output Frame %d"), ReadyFrame->GetId()));
+								SCOPE_CYCLE_COUNTER(STAT_MediaCapture_AnyThread_LockResource)
+
 								UMediaCapture::FMediaCaptureResourceData ResourceData;
 								ResourceData.Buffer = ColorDataBuffer;
 								ResourceData.Width = MediaCapture->DesiredOutputSize.X;
@@ -738,12 +768,13 @@ namespace UE::MediaCaptureData
 								MediaCapture->OnFrameCaptured_AnyThread(ReadyFrame->CaptureBaseData, ReadyFrame->UserData, MoveTemp(ResourceData));
 							}
 
+							ReadyFrame->Unlock_Unsafe();
 						}
-						ReadyFrame->Unlock_Unsafe(GPUMask);
 					}
-
-					ReadyFrame->bReadbackRequested = false;
 				}
+				
+				ReadyFrame->bDoingGPUCopy = false;
+				ReadyFrame->bReadbackRequested = false;
 
 				UE_LOG(LogMediaIOCore, Verbose, TEXT("[%s - %s] - Completed pending frame %d."), *MediaCapture->MediaOutputName, *FThreadManager::GetThreadName(FPlatformTLS::GetCurrentThreadId()), ReadyFrame->GetId());
 				MediaCapture->FrameManager->CompleteNextPending(*ReadyFrame);
@@ -852,7 +883,7 @@ namespace UE::MediaCaptureData
 										{
 											++MediaCapture->WaitingForRenderCommandExecutionCounter;
 
-											ENQUEUE_RENDER_COMMAND(MediaOutputCaptureFrameCreateResources)([MediaCapture, CapturingFrame](FRHICommandList& RHICommandList)
+											ENQUEUE_RENDER_COMMAND(MediaOutputCaptureReadbackComplete)([MediaCapture, CapturingFrame](FRHICommandList& RHICommandList)
 											{
 												OnReadbackComplete(RHICommandList, MediaCapture, CapturingFrame);
 											});
@@ -968,7 +999,7 @@ namespace UE::MediaCaptureData
 					RDG_GPU_STAT_SCOPE(Args.GraphBuilder, MediaCapture_Readback)
 					TRACE_CPUPROFILER_EVENT_SCOPE(UMediaCapture::EnqueueReadback);
 
-					CapturingFrame->EnqueueCopy(Args.GraphBuilder, OutputResource);
+					CapturingFrame->EnqueueCopy(Args.GraphBuilder, OutputResource, Args.MediaCapture->UseAnyThreadCapture());
 					CapturingFrame->bReadbackRequested = true;
 				}
 				else
@@ -2239,7 +2270,7 @@ bool UMediaCapture::ProcessReadyFrame_RenderThread(FRHICommandListImmediate& RHI
 			if (bIsReadbackReady)
 			{
 				SCOPE_CYCLE_COUNTER(STAT_MediaCapture_RenderThread_LockResource);
-				ColorDataBuffer = ReadyFrame->Lock(RHICmdList, GPUMask, RowStride);
+				ColorDataBuffer = ReadyFrame->Lock(RHICmdList, RowStride);
 			}
 			else
 			{
