Disable experimental scheduling in DX11 since our d3d11 fences rely on frame numbers which ends up clashing with the media capture sync point pass waiting mechanisms

diff --git a/Engine/Plugins/Media/MediaIOFramework/Source/MediaIOCore/Private/MediaCapture.cpp b/Engine/Plugins/Media/MediaIOFramework/Source/MediaIOCore/Private/MediaCapture.cpp
--- a/Engine/Plugins/Media/MediaIOFramework/Source/MediaIOCore/Private/MediaCapture.cpp
+++ b/Engine/Plugins/Media/MediaIOFramework/Source/MediaIOCore/Private/MediaCapture.cpp
@@ -1010,7 +1010,7 @@ namespace UE::MediaCaptureData
 				Args.MediaCapture->FrameManager->MarkPending(*CapturingFrame);
 				++Args.MediaCapture->PendingFrameCount;
 
-				if (CVarMediaIOEnableExperimentalScheduling.GetValueOnAnyThread() == 1)
+				if (Args.MediaCapture->UseExperimentalScheduling())
 				{
 					if (CapturingFrame->IsTextureResource())
 					{
@@ -1250,7 +1250,7 @@ bool UMediaCapture::StartSourceCapture(TSharedPtr<UE::MediaCapture::Private::FCa
 
 	if (!GRHISupportsMultithreading)
 	{
-		UE_LOG(LogMediaIOCore, Display, TEXT("AnyThread Capture was disabled because the current RHI does not support Multithreading."));
+		UE_LOG(LogMediaIOCore, Display, TEXT("Experimental scheduling and AnyThread Capture was disabled because the current RHI does not support Multithreading."));
 	}
 	else if (!SupportsAnyThreadCapture())
 	{
@@ -1384,6 +1384,11 @@ bool UMediaCapture::UpdateSource(TSharedPtr<UE::MediaCapture::Private::FCaptureS
 	return true;
 }
 
+bool UMediaCapture::UseExperimentalScheduling() const
+{
+	return GRHISupportsMultithreading && CVarMediaIOEnableExperimentalScheduling.GetValueOnAnyThread() == 1;
+}
+
 bool UMediaCapture::UseAnyThreadCapture() const
 {
 	return GRHISupportsMultithreading
@@ -1645,7 +1650,7 @@ void UMediaCapture::CaptureImmediate_RenderThread(const UE::MediaCaptureData::FC
 	}
 
 	// If CVarMediaIOEnableExperimentalScheduling is enabled, it means that a sync point pass was added, and the ready frame processing will be done later
-	if (CVarMediaIOEnableExperimentalScheduling.GetValueOnAnyThread() == 0)
+	if (!UseExperimentalScheduling())
 	{
 		if (TSharedPtr<FCaptureFrame> NextPending = FrameManager->PeekNextPending<FCaptureFrame>())
 		{
@@ -1731,7 +1736,7 @@ void UMediaCapture::ResetFixedViewportSize(TSharedPtr<FSceneViewport> InViewport
 
 void UMediaCapture::WaitForSingleExperimentalSchedulingTaskToComplete()
 {
-	if (CVarMediaIOEnableExperimentalScheduling.GetValueOnAnyThread())
+	if (UseExperimentalScheduling())
 	{
 		// Presumably the rendering thread could be in the process of dispatching the task
 		FScopeLock ScopedLock(&PendingReadbackTasksCriticalSection);
@@ -1779,7 +1784,7 @@ void UMediaCapture::CleanupCompletedExperimentalSchedulingTasks()
 
 void UMediaCapture::WaitForAllExperimentalSchedulingTasksToComplete()
 {
-	if (CVarMediaIOEnableExperimentalScheduling.GetValueOnAnyThread())
+	if (UseExperimentalScheduling())
 	{
 		FScopeLock ScopedLock(&PendingReadbackTasksCriticalSection);
 
@@ -1917,7 +1922,7 @@ void UMediaCapture::PrepareAndDispatchCapture_GameThread(const TSharedPtr<UE::Me
 					GraphBuilder.Execute();
 				}
 				
-				if (CVarMediaIOEnableExperimentalScheduling.GetValueOnAnyThread() == 0)
+				if (!InMediaCapture->UseExperimentalScheduling())
 				{
 					//Process the next pending frame
 					UE_LOG(LogMediaIOCore, Verbose, TEXT("[%s] - Processing pending frame %d"), *InMediaCapture->MediaOutputName, NextPending ? NextPending->FrameId : -1);
diff --git a/Engine/Plugins/Media/MediaIOFramework/Source/MediaIOCore/Public/MediaCapture.h b/Engine/Plugins/Media/MediaIOFramework/Source/MediaIOCore/Public/MediaCapture.h
--- a/Engine/Plugins/Media/MediaIOFramework/Source/MediaIOCore/Public/MediaCapture.h
+++ b/Engine/Plugins/Media/MediaIOFramework/Source/MediaIOCore/Public/MediaCapture.h
@@ -571,6 +571,16 @@ private:
 	bool StartSourceCapture(TSharedPtr<UE::MediaCapture::Private::FCaptureSource> InSource);
 	bool UpdateSource(TSharedPtr<UE::MediaCapture::Private::FCaptureSource> InCaptureSource);
 
+	/** 
+	 * Whether capture should be done using experimental scheduling.
+	 * Experimental scheduling means that readback is started on the render thread but we don't block on the game thread if it hasn't been completed by the next frame.
+	 */
+	bool UseExperimentalScheduling() const;
+
+	/**
+	 * When using anythread capture, the capture callback will be called directly after readback is complete, meaning there will be minimal delay.
+	 * When not using it, the capture callback will be called on the render thread which can introduce a delay.
+	 */
 	bool UseAnyThreadCapture() const;
 
 protected:
