Fixing up 10-bit YUV capture for Aja/BlackMagic cards (MediaIOCore player) which was broken since the pro-codec support work
#rnx

diff --git a/Engine/Plugins/Media/MediaIOFramework/Source/MediaIOCore/Public/MediaIOCoreTextureSampleBase.h b/Engine/Plugins/Media/MediaIOFramework/Source/MediaIOCore/Public/MediaIOCoreTextureSampleBase.h
--- a/Engine/Plugins/Media/MediaIOFramework/Source/MediaIOCore/Public/MediaIOCoreTextureSampleBase.h
+++ b/Engine/Plugins/Media/MediaIOFramework/Source/MediaIOCore/Public/MediaIOCoreTextureSampleBase.h
@@ -129,8 +129,9 @@ public:
 		case EMediaTextureSampleFormat::CharYVYU:
 			return FIntPoint(Width / 2, Height);
 		case EMediaTextureSampleFormat::YUVv210:
+			// Data for 6 output pixels is contained in 4 actual texture pixels
 			// Padding aligned on 48 (16 and 6 at the same time)
-			return FIntPoint((((Width + 47) / 48) * 48) / 6, Height);
+			return FIntPoint(4 * ((((Width + 47) / 48) * 48) / 6), Height);
 		default:
 			return FIntPoint(Width, Height);
 		}
