Fix media capture crash when closing editor with VCam active

diff --git a/Engine/Plugins/Media/MediaIOFramework/Source/MediaIOCore/Private/MediaCapture.cpp b/Engine/Plugins/Media/MediaIOFramework/Source/MediaIOCore/Private/MediaCapture.cpp
--- a/Engine/Plugins/Media/MediaIOFramework/Source/MediaIOCore/Private/MediaCapture.cpp
+++ b/Engine/Plugins/Media/MediaIOFramework/Source/MediaIOCore/Private/MediaCapture.cpp
@@ -1342,46 +1342,54 @@ bool UMediaCapture::ProcessCapture_RenderThread(const TSharedPtr<UE::MediaCaptur
 
 	if (CapturingFrame)
 	{
-		// Register the source texture with the render graph.
-		if (!Args.RDGResourceToCapture)
-		{
-			// If we weren't passed a rdg texture, register the external rhi texture.
-			Args.RDGResourceToCapture = Args.GraphBuilder.RegisterExternalTexture(CreateRenderTarget(Args.ResourceToCapture, TEXT("MediaCaptureSourceTexture")));
-		}
-
-		Args.MediaCapture->CaptureRenderPipeline->InitializeResources_RenderThread(CapturingFrame, Args.GraphBuilder, Args.RDGResourceToCapture);
-
-		// Call the capture frame algo based on the specific type of resource we are using
 		bool bHasCaptureSuceeded = false;
 		bool bBeforeCaptureSuccess = true;
-		if (DesiredOutputResourceType == EMediaCaptureResourceType::Texture)
+
+		if (Args.RDGResourceToCapture || Args.ResourceToCapture)
 		{
-			if (ensure(CapturingFrame->IsTextureResource()))
+			// Register the source texture with the render graph.
+			if (!Args.RDGResourceToCapture)
 			{
-				Args.MediaCapture->BeforeFrameCaptured_RenderingThread(CapturingFrame->CaptureBaseData, CapturingFrame->UserData, CapturingFrame->GetTextureResource());
+				// If we weren't passed a rdg texture, register the external rhi texture.
+				Args.RDGResourceToCapture = Args.GraphBuilder.RegisterExternalTexture(CreateRenderTarget(Args.ResourceToCapture, TEXT("MediaCaptureSourceTexture")));
 			}
-			else
+
+			Args.MediaCapture->CaptureRenderPipeline->InitializeResources_RenderThread(CapturingFrame, Args.GraphBuilder, Args.RDGResourceToCapture);
+
+			// Call the capture frame algo based on the specific type of resource we are using
+			if (DesiredOutputResourceType == EMediaCaptureResourceType::Texture)
 			{
-				bBeforeCaptureSuccess = false;
-				UE_LOG(LogMediaIOCore, Error, TEXT("The capture will stop for '%s'. Capture frame was expected to use Texture resource but wasn't."), *Args.MediaCapture->MediaOutputName);
+				if (ensure(CapturingFrame->IsTextureResource()))
+				{
+					Args.MediaCapture->BeforeFrameCaptured_RenderingThread(CapturingFrame->CaptureBaseData, CapturingFrame->UserData, CapturingFrame->GetTextureResource());
+				}
+				else
+				{
+					bBeforeCaptureSuccess = false;
+					UE_LOG(LogMediaIOCore, Error, TEXT("The capture will stop for '%s'. Capture frame was expected to use Texture resource but wasn't."), *Args.MediaCapture->MediaOutputName);
+				}
 			}
-		}
-		else
-		{
-			if (ensure(CapturingFrame->IsBufferResource()))
+			else
 			{
-				Args.MediaCapture->BeforeFrameCaptured_RenderingThread(CapturingFrame->CaptureBaseData, CapturingFrame->UserData, CapturingFrame->GetBufferResource());
+				if (ensure(CapturingFrame->IsBufferResource()))
+				{
+					Args.MediaCapture->BeforeFrameCaptured_RenderingThread(CapturingFrame->CaptureBaseData, CapturingFrame->UserData, CapturingFrame->GetBufferResource());
+				}
+				else
+				{
+					bBeforeCaptureSuccess = false;
+					UE_LOG(LogMediaIOCore, Error, TEXT("The capture will stop for '%s'. Capture frame was expected to use Buffer resource but wasn't."), *Args.MediaCapture->MediaOutputName);
+				}
 			}
-			else
+
+			if (bBeforeCaptureSuccess)
 			{
-				bBeforeCaptureSuccess = false;
-				UE_LOG(LogMediaIOCore, Error, TEXT("The capture will stop for '%s'. Capture frame was expected to use Buffer resource but wasn't."), *Args.MediaCapture->MediaOutputName);
+				bHasCaptureSuceeded = UE::MediaCaptureData::FMediaCaptureHelper::CaptureFrame(Args, CapturingFrame);
 			}
 		}
-
-		if (bBeforeCaptureSuccess)
+		else
 		{
-			bHasCaptureSuceeded = UE::MediaCaptureData::FMediaCaptureHelper::CaptureFrame(Args, CapturingFrame);
+			UE_LOG(LogMediaIOCore, Error, TEXT("The capture will stop for '%s'. No texture was acquired for the capture."), *Args.MediaCapture->MediaOutputName);
 		}
 		
 		if (bHasCaptureSuceeded == false)
