Disable media capture resampling to fix CIS error

diff --git a/Engine/Plugins/Media/MediaIOFramework/Source/MediaIOCore/Private/MediaCaptureRenderPass.cpp b/Engine/Plugins/Media/MediaIOFramework/Source/MediaIOCore/Private/MediaCaptureRenderPass.cpp
--- a/Engine/Plugins/Media/MediaIOFramework/Source/MediaIOCore/Private/MediaCaptureRenderPass.cpp
+++ b/Engine/Plugins/Media/MediaIOFramework/Source/MediaIOCore/Private/MediaCaptureRenderPass.cpp
@@ -84,6 +84,7 @@ namespace UE::MediaCapture::Resample
 	using FResamplePermutationDomain = TShaderPermutationDomain<FResampleQualityDimension>;
 
 	/** Resample pixel shader struct copied from PostProcessResample.h */
+	/*
 	class FMediaCaptureResamplePS : public FGlobalShader
 	{
 	public:
@@ -102,8 +103,9 @@ namespace UE::MediaCapture::Resample
 			return true;
 		}
 	};
+	*/
 
-	IMPLEMENT_GLOBAL_SHADER(FMediaCaptureResamplePS, "/Engine/Private/PostProcessDownsample.usf", "MainPS", SF_Pixel);
+	//IMPLEMENT_GLOBAL_SHADER(FMediaCaptureResamplePS, "/Engine/Private/PostProcessDownsample.usf", "MainPS", SF_Pixel);
 
 	/** Adds a Resample pass to the graph builder. */
 	void AddResamplePass(FRDGBuilder& GraphBuilder, const FViewInfo& View, const FResamplePassInputs& Inputs, FRDGTextureRef OutputTexture)
@@ -119,6 +121,8 @@ namespace UE::MediaCapture::Resample
 		FResamplePermutationDomain PermutationVector;
 		PermutationVector.Set<FResampleQualityDimension>(Inputs.Quality);
 
+		/*
+
 		FMediaCaptureResamplePS::FParameters* PassParameters = GraphBuilder.AllocParameters<FMediaCaptureResamplePS::FParameters>();
 		PassParameters->Common = GetResampleParameters(View, Output, Inputs.SceneColor, Inputs.Quality);
 		PassParameters->RenderTargets[0] = Output.GetRenderTargetBinding();
@@ -135,6 +139,8 @@ namespace UE::MediaCapture::Resample
 			PixelShader,
 			PassParameters,
 			Output.ViewRect);
+
+		*/
 	}
 
 }
@@ -415,7 +421,8 @@ namespace UE::MediaCapture
 
 		if (MediaCapture->GetDesiredCaptureOptions().ResizeMethod == EMediaCaptureResizeMethod::ResizeInRenderPass)
 		{
-			RenderPasses.Add(CreateResamplePass());
+			// Disabled temporarily to fix CIS issue related to the shader.
+			//RenderPasses.Add(CreateResamplePass());
 		}
 		
 		if (MediaCapture->GetDesiredCaptureOptions().ColorConversionSettings.IsValid())
