Fix crash when doing media capture with ndisplay

diff --git a/Engine/Plugins/Media/MediaIOFramework/Source/MediaIOCore/Private/MediaCaptureHelper.cpp b/Engine/Plugins/Media/MediaIOFramework/Source/MediaIOCore/Private/MediaCaptureHelper.cpp
--- a/Engine/Plugins/Media/MediaIOFramework/Source/MediaIOCore/Private/MediaCaptureHelper.cpp
+++ b/Engine/Plugins/Media/MediaIOFramework/Source/MediaIOCore/Private/MediaCaptureHelper.cpp
@@ -141,9 +141,9 @@ bool FMediaCaptureHelper::CaptureFrame(const UE::MediaCaptureData::FCaptureFrame
 		SCOPE_CYCLE_COUNTER(STAT_MediaCaptureHelper_RenderThread_FrameCapture);
 
 		FRDGTextureRef SourceRGBTexture = Args.RDGResourceToCapture;
-
 		// Final pass output resource used by the current capture method (texture or buffer)
-		FRDGViewableResource* FinalPassOutputResource = CapturingFrame->RenderPassResources.FinalRDGResource;
+		FRDGViewableResource* FinalPassOutputResource = CapturingFrame->RenderPassResources.GetFinalRDGResource(Args.GraphBuilder);
+		
 		{
 			TRACE_CPUPROFILER_EVENT_SCOPE(UMediaCapture::GraphSetup);
 			SCOPED_DRAW_EVENTF(Args.GraphBuilder.RHICmdList, MediaCapture, TEXT("MediaCapture"));
diff --git a/Engine/Plugins/Media/MediaIOFramework/Source/MediaIOCore/Private/MediaCaptureRenderPass.cpp b/Engine/Plugins/Media/MediaIOFramework/Source/MediaIOCore/Private/MediaCaptureRenderPass.cpp
--- a/Engine/Plugins/Media/MediaIOFramework/Source/MediaIOCore/Private/MediaCaptureRenderPass.cpp
+++ b/Engine/Plugins/Media/MediaIOFramework/Source/MediaIOCore/Private/MediaCaptureRenderPass.cpp
@@ -572,8 +572,6 @@ namespace UE::MediaCapture
 			}
 
 			RenderPassResources.LastRenderPassName = RenderPasses.Last().Name;
-			// InputResource now points to the output of the final render pass.
-			RenderPassResources.FinalRDGResource = InitializePassOutputArgs.InputResource;
 			CapturingFrame->RenderPassResources = MoveTemp(RenderPassResources);
 		}
 	}
diff --git a/Engine/Plugins/Media/MediaIOFramework/Source/MediaIOCore/Private/MediaCaptureRenderPass.h b/Engine/Plugins/Media/MediaIOFramework/Source/MediaIOCore/Private/MediaCaptureRenderPass.h
--- a/Engine/Plugins/Media/MediaIOFramework/Source/MediaIOCore/Private/MediaCaptureRenderPass.h
+++ b/Engine/Plugins/Media/MediaIOFramework/Source/MediaIOCore/Private/MediaCaptureRenderPass.h
@@ -120,24 +120,38 @@ namespace UE::MediaCapture
 	/** Data related to render passes that live on a capturing frame. **/
 	struct FRenderPassFrameResources
 	{
+		/** Get the render target used by the last render pass in the pipeline. */
 		TRefCountPtr<IPooledRenderTarget> GetLastRenderTarget() const
 		{
 			return PassTextureOutputs.FindRef(LastRenderPassName);
 		}
 
+		/** Get the buffer used by the last render pass in the pipeline. */
 		TRefCountPtr<FRDGPooledBuffer> GetLastBuffer() const
 		{
 			return PassBufferOutputs.FindRef(LastRenderPassName);
 		}
 
+		/** Register and retrieve the final resource used by this capture frame in the render passes. */
+		FRDGViewableResource* GetFinalRDGResource(FRDGBuilder& GraphBuilder) const
+		{
+			if (const TRefCountPtr<IPooledRenderTarget> FinalRT = GetLastRenderTarget())
+			{
+				return GraphBuilder.RegisterExternalTexture(FinalRT, FinalRT->GetDesc().DebugName);
+			}
+			else if (const TRefCountPtr<FRDGPooledBuffer> FinalBuffer = GetLastBuffer())
+			{
+				return GraphBuilder.RegisterExternalBuffer(FinalBuffer, FinalBuffer->GetName());
+			}
+			return nullptr;
+		}
+
 		/** Holds the name of the last render pass. */
 		FName LastRenderPassName;
 		/** Map of pass name to output render target. */
 		TMap<FName, TRefCountPtr<IPooledRenderTarget>> PassTextureOutputs;
 		/** Map of pass name to output buffer. */
 		TMap<FName, TRefCountPtr<FRDGPooledBuffer>> PassBufferOutputs;
-		/** The final pass' resource that was registered with the graph builder. */
-		FRDGViewableResource* FinalRDGResource;
 	};
 
 	/**
