- Fixing a crash in media capture when capture immediate is used and invalid resource is detected

diff --git a/Engine/Plugins/Media/MediaIOFramework/Source/MediaIOCore/Private/MediaCapture.cpp b/Engine/Plugins/Media/MediaIOFramework/Source/MediaIOCore/Private/MediaCapture.cpp
--- a/Engine/Plugins/Media/MediaIOFramework/Source/MediaIOCore/Private/MediaCapture.cpp
+++ b/Engine/Plugins/Media/MediaIOFramework/Source/MediaIOCore/Private/MediaCapture.cpp
@@ -593,6 +593,15 @@ void UMediaCapture::CaptureImmediate_RenderThread(FRDGBuilder& GraphBuilder, FRH
 
 void UMediaCapture::CaptureImmediate_RenderThread(FRDGBuilder& GraphBuilder,  FRDGTextureRef InSourceTextureRef, FIntRect SourceViewRect)
 {
+	// Make sure captureimmediate runs are seen by game thread when stopping is requested. Having a pending render command will cause a flush
+	// hence letting the capture finish. If stopping is in progress, stopped state will be seen by the captureimmediate and early exit.
+	++WaitingForRenderCommandExecutionCounter;
+
+	ON_SCOPE_EXIT
+	{
+		--WaitingForRenderCommandExecutionCounter;
+	};
+	
 	UE::MediaCaptureData::FCaptureFrameArgs CaptureArgs{GraphBuilder};
 	CaptureArgs.MediaCapture = this;
 	CaptureArgs.DesiredSize = DesiredSize;
@@ -636,9 +645,6 @@ void UMediaCapture::CaptureImmediate_RenderThread(const UE::MediaCaptureData::FC
 		UE_LOG(LogMediaIOCore, Error, TEXT("[%s] - Trying to capture a RHI resource with another capture type."), *MediaOutputName);
 		SetState(EMediaCaptureState::Error);
 	}
-	
-	// Keep resource size up to date with incoming resource to capture
-    StaticCastSharedPtr<UE::MediaCapture::Private::FRHIResourceCaptureSource>(CaptureSource)->ResourceDescription.ResourceSize = Args.GetSizeXY();
 
 	if (GetState() != EMediaCaptureState::Capturing && GetState() != EMediaCaptureState::StopRequested)
 	{
@@ -650,6 +656,16 @@ void UMediaCapture::CaptureImmediate_RenderThread(const UE::MediaCaptureData::FC
 		return;
 	}
 
+	if (CaptureSource.IsValid() == false)
+	{
+		UE_LOG(LogMediaIOCore, Error, TEXT("[%s] - Trying to capture a RHI resource with an invalid source."), *MediaOutputName);
+		SetState(EMediaCaptureState::Error);
+		return;
+	}
+
+	// Keep resource size up to date with incoming resource to capture
+ 	StaticCastSharedPtr<UE::MediaCapture::Private::FRHIResourceCaptureSource>(CaptureSource)->ResourceDescription.ResourceSize = Args.GetSizeXY();
+
 	// Get cached capture data from game thread. We want to find a cached frame matching current render thread frame number
 	bool bFoundMatchingData = false;
 	FQueuedCaptureData NextCaptureData;
diff --git a/Engine/Plugins/Media/MediaIOFramework/Source/MediaIOCore/Public/MediaCapture.h b/Engine/Plugins/Media/MediaIOFramework/Source/MediaIOCore/Public/MediaCapture.h
--- a/Engine/Plugins/Media/MediaIOFramework/Source/MediaIOCore/Public/MediaCapture.h
+++ b/Engine/Plugins/Media/MediaIOFramework/Source/MediaIOCore/Public/MediaCapture.h
@@ -653,7 +653,7 @@ private:
 	TPimplPtr<UE::MediaCaptureData::FFrameManager> FrameManager;
 	int32 NumberOfCaptureFrame = 2;
 	int32 CaptureRequestCount = 0;
-	EMediaCaptureState MediaState = EMediaCaptureState::Stopped;
+	std::atomic<EMediaCaptureState> MediaState = EMediaCaptureState::Stopped;
 
 	FCriticalSection AccessingCapturingSource;
 
