Add additional deinterlacing methods and ability to provide your own deinterlacer for video input

diff --git a/Engine/Plugins/Media/AjaMedia/Source/AjaMedia/Private/Player/AjaMediaPlayer.cpp b/Engine/Plugins/Media/AjaMedia/Source/AjaMedia/Private/Player/AjaMediaPlayer.cpp
--- a/Engine/Plugins/Media/AjaMedia/Source/AjaMedia/Private/Player/AjaMediaPlayer.cpp
+++ b/Engine/Plugins/Media/AjaMedia/Source/AjaMedia/Private/Player/AjaMediaPlayer.cpp
@@ -13,6 +13,7 @@
 
 #include "Engine/Engine.h"
 #include "GenlockedCustomTimeStep.h"
+#include "MediaIOCoreDeinterlacer.h"
 #include "MediaIOCoreEncodeTime.h"
 #include "MediaIOCoreFileWriter.h"
 #include "MediaIOCoreSamples.h"
@@ -79,8 +80,8 @@ FAjaMediaPlayer::FAjaMediaPlayer(IMediaEventSink& InEventSink)
 	, InputChannel(nullptr)
 	, SupportedSampleTypes(EMediaIOSampleType::None)
 	, bPauseRequested(false)
+	, DeviceProvider(MakePimpl<FAjaDeviceProvider>())
 {
-	DeviceProvider = MakePimpl<FAjaDeviceProvider>();
 }
 
 
@@ -658,10 +659,9 @@ bool FAjaMediaPlayer::OnInputFrameReceived(const AJA::AJAInputFrameData& InInput
 			}
 			else
 			{
-				auto TextureSample = TextureSamplePool->AcquireShared();
 				if (InVideoFrame.bIsProgressivePicture)
-
 				{
+					auto TextureSample = TextureSamplePool->AcquireShared();
 					if (TextureSample->InitializeProgressive(InVideoFrame, VideoSampleFormat, DecodedTime, VideoFrameRate, DecodedTimecode, bIsSRGBInput))
 					{
 						Samples->AddVideo(TextureSample);
@@ -669,20 +669,30 @@ bool FAjaMediaPlayer::OnInputFrameReceived(const AJA::AJAInputFrameData& InInput
 				}
 				else
 				{
-					bool bEven = true;
-					if (UGenlockedCustomTimeStep* TimeStep = Cast< UGenlockedCustomTimeStep>(GEngine->GetCustomTimeStep()))
+					if (UGenlockedCustomTimeStep* TimeStep = Cast<UGenlockedCustomTimeStep>(GEngine->GetCustomTimeStep()))
 					{
 						bAddedSynchronizedSamples = TimeStep->GetSynchronizationState() == ECustomTimeStepSynchronizationState::Synchronized;
 					}
-					if (TextureSample->InitializeInterlaced_Halfed(InVideoFrame, VideoSampleFormat, DecodedTime, VideoFrameRate, DecodedTimecode, bEven, bIsSRGBInput))
+
+					UE::MediaIOCore::FVideoFrame FrameInfo
 					{
-						Samples->AddVideo(TextureSample);
-					}
-					bEven = !bEven;
-					auto TextureSampleOdd = TextureSamplePool->AcquireShared();
-					if (TextureSampleOdd->InitializeInterlaced_Halfed(InVideoFrame, VideoSampleFormat, DecodedTimeF2, VideoFrameRate, DecodedTimecodeF2, bEven, bIsSRGBInput))
+						InVideoFrame.VideoBuffer,
+						InVideoFrame.VideoBufferSize,
+						InVideoFrame.Stride,
+						InVideoFrame.Width,
+						InVideoFrame.Height,
+						VideoSampleFormat,
+						DecodedTime,
+						VideoFrameRate,
+						DecodedTimecode,
+						bIsSRGBInput ? UE::Color::EColorSpace::sRGB : UE::Color::EColorSpace::None
+					};
+
+					const TArray<TSharedRef<FMediaIOCoreTextureSampleBase>> DeinterlacedSamples = Deinterlace(FrameInfo);
+
+					for (const TSharedRef<FMediaIOCoreTextureSampleBase>& TextureSample : DeinterlacedSamples)
 					{
-						Samples->AddVideo(TextureSampleOdd);
+						Samples->AddVideo(TextureSample);
 					}
 				}
 			}
@@ -974,4 +984,9 @@ bool FAjaMediaPlayer::SetRate(float Rate)
 	return false;
 }
 
+TSharedPtr<FMediaIOCoreTextureSampleBase> FAjaMediaPlayer::AcquireSample_AnyThread() const
+{
+	return TextureSamplePool->AcquireShared();
+}
+
 #undef LOCTEXT_NAMESPACE
diff --git a/Engine/Plugins/Media/AjaMedia/Source/AjaMedia/Private/Player/AjaMediaPlayer.h b/Engine/Plugins/Media/AjaMedia/Source/AjaMedia/Private/Player/AjaMediaPlayer.h
--- a/Engine/Plugins/Media/AjaMedia/Source/AjaMedia/Private/Player/AjaMediaPlayer.h
+++ b/Engine/Plugins/Media/AjaMedia/Source/AjaMedia/Private/Player/AjaMediaPlayer.h
@@ -29,6 +29,7 @@ namespace AJA
 	class AJAInputChannel;
 }
 
+
 /**
  * Implements a media player using AJA.
  *
@@ -104,20 +105,19 @@ protected:
 
 	virtual bool IsHardwareReady() const override;
 
-	//~ FMediaIOCorePlayerBase interface
+	//~ Begin FMediaIOCorePlayerBase interface
 	virtual void SetupSampleChannels() override;
-
 	virtual uint32 GetNumVideoFrameBuffers() const override
 	{
 		return MaxNumVideoFrameBuffer;
 	}
-
 	virtual EMediaIOCoreColorFormat GetColorFormat() const override
 	{
 		return AjaColorFormat == EAjaMediaSourceColorFormat::YUV2_8bit ? EMediaIOCoreColorFormat::YUV8 : EMediaIOCoreColorFormat::YUV10;
 	}
-
 	virtual void AddVideoSample(const TSharedRef<FMediaIOCoreTextureSampleBase>& InSample) override;
+	virtual TSharedPtr<FMediaIOCoreTextureSampleBase> AcquireSample_AnyThread() const override;
+	//~ End FMediaIOCorePlayerBase interface
 
 private:
 	bool Open_Internal(const FString& Url, const IMediaOptions* Options, AJA::AJAInputOutputChannelOptions AjaOptions);
diff --git a/Engine/Plugins/Media/AjaMedia/Source/AjaMedia/Public/AjaMediaSource.h b/Engine/Plugins/Media/AjaMedia/Source/AjaMedia/Public/AjaMediaSource.h
--- a/Engine/Plugins/Media/AjaMedia/Source/AjaMedia/Public/AjaMediaSource.h
+++ b/Engine/Plugins/Media/AjaMedia/Source/AjaMedia/Public/AjaMediaSource.h
@@ -2,7 +2,7 @@
 
 #pragma once
 
-#include "TimeSynchronizableMediaSource.h"
+#include "CaptureCardMediaSource.h"
 #include "MediaIOCoreDefinitions.h"
 
 #include "AjaMediaSource.generated.h"
@@ -31,7 +31,7 @@ enum class EAjaMediaAudioChannel : uint8
  * Media source for AJA streams.
  */
 UCLASS(BlueprintType, hideCategories=(Platforms,Object), meta=(MediaIOCustomLayout="AJA"))
-class AJAMEDIA_API UAjaMediaSource : public UTimeSynchronizableMediaSource
+class AJAMEDIA_API UAjaMediaSource : public UCaptureCardMediaSource
 {
 	GENERATED_BODY()
 
diff --git a/Engine/Plugins/Media/BlackmagicMedia/Source/BlackmagicMedia/Private/Player/BlackmagicMediaPlayer.cpp b/Engine/Plugins/Media/BlackmagicMedia/Source/BlackmagicMedia/Private/Player/BlackmagicMediaPlayer.cpp
--- a/Engine/Plugins/Media/BlackmagicMedia/Source/BlackmagicMedia/Private/Player/BlackmagicMediaPlayer.cpp
+++ b/Engine/Plugins/Media/BlackmagicMedia/Source/BlackmagicMedia/Private/Player/BlackmagicMediaPlayer.cpp
@@ -11,6 +11,7 @@
 #include "IBlackmagicMediaModule.h"
 #include "IMediaEventSink.h"
 #include "IMediaOptions.h"
+#include "MediaIOCoreDeinterlacer.h"
 #include "MediaIOCoreEncodeTime.h"
 #include "MediaIOCoreFileWriter.h"
 #include "MediaIOCoreSamples.h"
@@ -354,39 +355,25 @@ namespace BlackmagicMediaPlayerHelpers
 					}
 					else
 					{
-						bool bEven = true;
-
-						auto TextureSampleEven = MediaPlayer->TextureSamplePool->AcquireShared();
-
-						if (TextureSampleEven->InitializeWithEvenOddLine(bEven
-							, InFrameInfo.VideoBuffer
-							, InFrameInfo.VideoPitch * InFrameInfo.VideoHeight
-							, InFrameInfo.VideoPitch
-							, InFrameInfo.VideoWidth
-							, InFrameInfo.VideoHeight
-							, SampleFormat
-							, DecodedTime
-							, MediaPlayer->VideoFrameRate
-							, DecodedTimecode
-							, bIsSRGBInput))
+						UE::MediaIOCore::FVideoFrame FrameInfo
 						{
-							MediaPlayer->Samples->AddVideo(TextureSampleEven);
-						}
-
-						auto TextureSampleOdd = MediaPlayer->TextureSamplePool->AcquireShared();
-						if (TextureSampleOdd->InitializeWithEvenOddLine(!bEven
-							, InFrameInfo.VideoBuffer
-							, InFrameInfo.VideoPitch * InFrameInfo.VideoHeight
-							, InFrameInfo.VideoPitch
-							, InFrameInfo.VideoWidth
-							, InFrameInfo.VideoHeight
-							, SampleFormat
-							, DecodedTimeF2
-							, MediaPlayer->VideoFrameRate
-							, DecodedTimecodeF2
-							, bIsSRGBInput))
+							InFrameInfo.VideoBuffer,
+							(uint32)InFrameInfo.VideoPitch * InFrameInfo.VideoHeight,
+							(uint32)InFrameInfo.VideoPitch,
+							(uint32)InFrameInfo.VideoWidth,
+							(uint32)InFrameInfo.VideoHeight,
+							SampleFormat,
+							DecodedTime,
+							MediaPlayer->VideoFrameRate,
+							DecodedTimecode,
+							bIsSRGBInput ? UE::Color::EColorSpace::sRGB : UE::Color::EColorSpace::None
+						};
+
+						const TArray<TSharedRef<FMediaIOCoreTextureSampleBase>> DeinterlacedSamples = MediaPlayer->Deinterlace(FrameInfo);
+
+						for (const TSharedRef<FMediaIOCoreTextureSampleBase>& TextureSample : DeinterlacedSamples)
 						{
-							MediaPlayer->Samples->AddVideo(TextureSampleOdd);
+							MediaPlayer->Samples->AddVideo(TextureSample);
 						}
 					}
 				}
diff --git a/Engine/Plugins/Media/BlackmagicMedia/Source/BlackmagicMedia/Public/BlackmagicMediaSource.h b/Engine/Plugins/Media/BlackmagicMedia/Source/BlackmagicMedia/Public/BlackmagicMediaSource.h
--- a/Engine/Plugins/Media/BlackmagicMedia/Source/BlackmagicMedia/Public/BlackmagicMediaSource.h
+++ b/Engine/Plugins/Media/BlackmagicMedia/Source/BlackmagicMedia/Public/BlackmagicMediaSource.h
@@ -2,7 +2,7 @@
 
 #pragma once
 
-#include "TimeSynchronizableMediaSource.h"
+#include "CaptureCardMediaSource.h"
 
 #include "BlackmagicDeviceProvider.h"
 #include "MediaIOCoreDefinitions.h"
@@ -34,7 +34,7 @@ enum class EBlackmagicMediaAudioChannel : uint8
  * Media source description for Blackmagic.
  */
 UCLASS(BlueprintType, hideCategories=(Platforms,Object), meta=(MediaIOCustomLayout="Blackmagic"))
-class BLACKMAGICMEDIA_API UBlackmagicMediaSource : public UTimeSynchronizableMediaSource
+class BLACKMAGICMEDIA_API UBlackmagicMediaSource : public UCaptureCardMediaSource
 {
 	GENERATED_BODY()
 
diff --git a/Engine/Plugins/Media/MediaIOFramework/Source/MediaIOCore/Private/CaptureCardMediaSource.cpp b/Engine/Plugins/Media/MediaIOFramework/Source/MediaIOCore/Private/CaptureCardMediaSource.cpp
new file mode 100644
--- /dev/null
+++ b/Engine/Plugins/Media/MediaIOFramework/Source/MediaIOCore/Private/CaptureCardMediaSource.cpp
@@ -0,0 +1,37 @@
+// Copyright Epic Games, Inc. All Rights Reserved.
+
+#include "CaptureCardMediaSource.h"
+
+int64 UCaptureCardMediaSource::GetMediaOption(const FName& Key, int64 DefaultValue) const
+{
+	if (Key == UE::CaptureCardMediaSource::InterlaceFieldOrder)
+	{
+		return static_cast<int64>(InterlaceFieldOrder);
+	}
+
+	return UTimeSynchronizableMediaSource::GetMediaOption(Key, DefaultValue);
+}
+
+FString UCaptureCardMediaSource::GetMediaOption(const FName& Key, const FString& DefaultValue) const
+{
+	if (Key == UE::CaptureCardMediaSource::Deinterlacer)
+	{
+		if (Deinterlacer)
+		{
+			return Deinterlacer->GetPathName();
+		}
+	}
+
+	return UTimeSynchronizableMediaSource::GetMediaOption(Key, DefaultValue);
+}
+
+bool UCaptureCardMediaSource::HasMediaOption(const FName& Key) const
+{
+	if (Key == UE::CaptureCardMediaSource::Deinterlacer
+		|| Key == UE::CaptureCardMediaSource::InterlaceFieldOrder)
+		{
+			return true;
+		}
+
+	return UTimeSynchronizableMediaSource::HasMediaOption(Key);
+}
\ No newline at end of file
diff --git a/Engine/Plugins/Media/MediaIOFramework/Source/MediaIOCore/Private/MediaIOCoreDeinterlacer.cpp b/Engine/Plugins/Media/MediaIOFramework/Source/MediaIOCore/Private/MediaIOCoreDeinterlacer.cpp
new file mode 100644
--- /dev/null
+++ b/Engine/Plugins/Media/MediaIOFramework/Source/MediaIOCore/Private/MediaIOCoreDeinterlacer.cpp
@@ -0,0 +1,153 @@
+// Copyright Epic Games, Inc. All Rights Reserved.
+
+#include "MediaIOCoreDeinterlacer.h"
+
+#include "Async/ParallelFor.h"
+#include "ColorManagementDefines.h"
+
+namespace UE::MediaIOCore
+{
+	TArray<TSharedRef<FMediaIOCoreTextureSampleBase>> FDeinterlacer::Deinterlace(const FVideoFrame& InVideoFrame) const
+ 	{
+		TArray<TSharedRef<FMediaIOCoreTextureSampleBase>> Samples;
+
+		const TSharedPtr<FMediaIOCoreTextureSampleBase> TextureSample = AcquireSampleDelegate.Execute();
+
+		if (TextureSample && TextureSample->Initialize(InVideoFrame.VideoBuffer
+			, InVideoFrame.BufferSize
+			, InVideoFrame.Stride
+			, InVideoFrame.Width
+			, InVideoFrame.Height
+			, InVideoFrame.SampleFormat
+			, InVideoFrame.Time
+			, InVideoFrame.FrameRate
+			, InVideoFrame.Timecode
+			, InVideoFrame.ColorSpace == Color::EColorSpace::sRGB ? true : false))
+		{
+			Samples.Add(TextureSample.ToSharedRef());
+		}
+
+		return Samples;
+ 	}
+
+	TArray<TSharedRef<FMediaIOCoreTextureSampleBase>> FBobDeinterlacer::Deinterlace(const FVideoFrame& InVideoFrame) const
+	{
+		TArray<TSharedRef<FMediaIOCoreTextureSampleBase>> Samples;
+
+		bool bEven = true;
+
+		const TSharedPtr<FMediaIOCoreTextureSampleBase> TextureSampleEven = AcquireSampleDelegate.Execute();
+		const TSharedPtr<FMediaIOCoreTextureSampleBase> TextureSampleOdd = AcquireSampleDelegate.Execute();
+
+		if (TextureSampleEven && TextureSampleOdd)
+		{
+			if (TextureSampleEven->InitializeWithEvenOddLine(
+				  InterlaceFieldOrder == EMediaIOInterlaceFieldOrder::TopFieldFirst
+				, InVideoFrame.VideoBuffer
+				, InVideoFrame.BufferSize
+				, InVideoFrame.Stride
+				, InVideoFrame.Width
+				, InVideoFrame.Height
+				, InVideoFrame.SampleFormat
+				, InVideoFrame.Time
+				, InVideoFrame.FrameRate
+				, InVideoFrame.Timecode
+				, InVideoFrame.ColorSpace == Color::EColorSpace::sRGB ? true : false))
+			{
+				Samples.Add(TextureSampleEven.ToSharedRef());
+			}
+
+			// Don't create second sample if first one fails in order to avoid introducing field flipping.
+			if (Samples.Num() && TextureSampleOdd->InitializeWithEvenOddLine(
+				 !(InterlaceFieldOrder == EMediaIOInterlaceFieldOrder::TopFieldFirst)
+				, InVideoFrame.VideoBuffer
+				, InVideoFrame.BufferSize
+				, InVideoFrame.Stride
+				, InVideoFrame.Width
+				, InVideoFrame.Height
+				, InVideoFrame.SampleFormat
+				, InVideoFrame.Time
+				, InVideoFrame.FrameRate
+				, InVideoFrame.Timecode
+				, InVideoFrame.ColorSpace == Color::EColorSpace::sRGB ? true : false))
+			{
+				Samples.Add(TextureSampleOdd.ToSharedRef());
+			}
+		}
+
+		return Samples;
+	}
+
+	TArray<TSharedRef<FMediaIOCoreTextureSampleBase>> FBlendDeinterlacer::Deinterlace(const FVideoFrame& InVideoFrame) const
+	{
+		TArray<TSharedRef<FMediaIOCoreTextureSampleBase>> Samples;
+		
+		if (const TSharedPtr<FMediaIOCoreTextureSampleBase> TextureSample = AcquireSampleDelegate.Execute())
+		{
+			TArray<uint8> BlendedBuffer;
+			BlendedBuffer.AddZeroed(InVideoFrame.BufferSize / 2);
+
+			ParallelFor(InVideoFrame.Height / 2, [&](int32 LineIndex)
+			{
+				int32 SrcY = LineIndex * 2;
+				const uint8* Source1 = static_cast<const uint8*>(InVideoFrame.VideoBuffer) + (SrcY * InVideoFrame.Stride);
+				const uint8* Source2 = static_cast<const uint8*>(InVideoFrame.VideoBuffer) + ((SrcY + 1) * InVideoFrame.Stride);
+
+				uint8* Dest = BlendedBuffer.GetData() + LineIndex * InVideoFrame.Stride;
+				
+				for (uint32 XIndex = 0; XIndex < InVideoFrame.Stride; XIndex++)
+				{
+					*(Dest + XIndex) =  FMath::DivideAndRoundUp(Source1[XIndex] + Source2[XIndex], 2);
+				}
+			});
+
+			if (TextureSample->Initialize(
+				  MoveTemp(BlendedBuffer)
+				, InVideoFrame.Stride
+				, InVideoFrame.Width
+				, InVideoFrame.Height / 2
+				, InVideoFrame.SampleFormat
+				, InVideoFrame.Time
+				, InVideoFrame.FrameRate
+				, InVideoFrame.Timecode
+				, InVideoFrame.ColorSpace == Color::EColorSpace::sRGB ? true : false))
+			{
+				Samples.Add(TextureSample.ToSharedRef());
+			}
+		}
+
+		return Samples;
+	}
+
+	TArray<TSharedRef<FMediaIOCoreTextureSampleBase>> FDiscardDeinterlacer::Deinterlace(const FVideoFrame& InVideoFrame) const
+	{
+		TArray<TSharedRef<FMediaIOCoreTextureSampleBase>> Samples;
+		if (const TSharedPtr<FMediaIOCoreTextureSampleBase> TextureSample = AcquireSampleDelegate.Execute())
+		{
+			TArray<uint8> Buffer;
+			Buffer.Reserve(InVideoFrame.BufferSize / 2);
+
+			for (uint32 YIndex = InterlaceFieldOrder == EMediaIOInterlaceFieldOrder::TopFieldFirst ? 0 : 1; YIndex < InVideoFrame.Height; YIndex += 2)
+			{
+				const uint8* Source1 = reinterpret_cast<const uint8*>(InVideoFrame.VideoBuffer) + (YIndex * InVideoFrame.Stride);
+				Buffer.Append(Source1, InVideoFrame.Stride);
+			}
+
+			if (TextureSample->Initialize(
+				  MoveTemp(Buffer)
+				, InVideoFrame.Stride
+				, InVideoFrame.Width
+				, InVideoFrame.Height / 2
+				, InVideoFrame.SampleFormat
+				, InVideoFrame.Time
+				, InVideoFrame.FrameRate
+				, InVideoFrame.Timecode
+				, InVideoFrame.ColorSpace == Color::EColorSpace::sRGB ? true : false))
+			{
+				Samples.Add(TextureSample.ToSharedRef());
+			}
+		}
+
+		return Samples;
+	}
+}
\ No newline at end of file
diff --git a/Engine/Plugins/Media/MediaIOFramework/Source/MediaIOCore/Private/Player/MediaIOCorePlayerBase.cpp b/Engine/Plugins/Media/MediaIOFramework/Source/MediaIOCore/Private/Player/MediaIOCorePlayerBase.cpp
--- a/Engine/Plugins/Media/MediaIOFramework/Source/MediaIOCore/Private/Player/MediaIOCorePlayerBase.cpp
+++ b/Engine/Plugins/Media/MediaIOFramework/Source/MediaIOCore/Private/Player/MediaIOCorePlayerBase.cpp
@@ -2,6 +2,7 @@
 
 #include "MediaIOCorePlayerBase.h"
 
+#include "CaptureCardMediaSource.h"
 #include "Engine/Engine.h"
 #include "Engine/TimecodeProvider.h"
 #include "GPUTextureTransferModule.h"
@@ -11,6 +12,7 @@
 #include "IMediaModule.h"
 #include "IMediaPlayerFactory.h"
 #include "ITimeManagementModule.h"
+#include "MediaIOCoreDeinterlacer.h"
 #include "MediaIOCoreModule.h"
 #include "MediaIOCoreSamples.h"
 #include "Misc/App.h"
@@ -108,6 +110,8 @@ void FMediaIOCorePlayerBase::Close()
 {
 	if (CurrentState != EMediaState::Closed)
 	{
+		Deinterlacer.Reset();
+
 		CurrentState = EMediaState::Closed;
 		CurrentTime = FTimespan::Zero();
 		AudioTrackFormat.NumChannels = 0;
@@ -195,12 +199,28 @@ bool FMediaIOCorePlayerBase::Open(const FString& Url, const IMediaOptions* Optio
 	ITimeManagementModule::Get().GetTimedDataInputCollection().Add(this);
 
 	OpenUrl = Url;
-	bool bReadMediaOptions = ReadMediaOptions(Options);
+	const bool bReadMediaOptions = ReadMediaOptions(Options);
 
 	if (MediaIOCorePlayerDetail::CVarEnableGPUDirectInput.GetValueOnAnyThread())
 	{
 		CreateAndRegisterTextures(Options);
 	}
+	
+	const EMediaIOInterlaceFieldOrder InterlaceFieldOrder = static_cast<EMediaIOInterlaceFieldOrder>(Options->GetMediaOption(UE::CaptureCardMediaSource::InterlaceFieldOrder, (int64)EMediaIOInterlaceFieldOrder::TopFieldFirst));
+	const FString DefaultOption = TEXT("");
+	const FSoftObjectPath DeinterlacerPath = FSoftObjectPath(Options->GetMediaOption(UE::CaptureCardMediaSource::Deinterlacer, DefaultOption));
+	
+	UObject* DeinterlacerInstancer = DeinterlacerPath.ResolveObject();
+
+	TSharedPtr<UE::MediaIOCore::FDeinterlacer> DeinterlacerInstance;
+	if (DeinterlacerInstancer && DeinterlacerInstancer->IsA<UVideoDeinterlacer>())
+	{
+		Deinterlacer = Cast<UVideoDeinterlacer>(DeinterlacerInstancer)->Instantiate(UE::MediaIOCore::FDeinterlacer::FOnAcquireSample_AnyThread::CreateSP(this, &FMediaIOCorePlayerBase::AcquireSample_AnyThread), InterlaceFieldOrder);
+	}
+	else
+	{
+		Deinterlacer = MakeShared<UE::MediaIOCore::FDeinterlacer>(UE::MediaIOCore::FDeinterlacer::FOnAcquireSample_AnyThread::CreateSP(this, &FMediaIOCorePlayerBase::AcquireSample_AnyThread), InterlaceFieldOrder);
+	}
 
 	return bReadMediaOptions;
 }
@@ -743,6 +763,15 @@ void FMediaIOCorePlayerBase::ExecuteGPUTransfer(const TSharedPtr<FMediaIOCoreTex
 	});
 }
 
+TArray<TSharedRef<FMediaIOCoreTextureSampleBase>> FMediaIOCorePlayerBase::Deinterlace(const UE::MediaIOCore::FVideoFrame& InVideoFrame) const
+{
+	if (Deinterlacer)
+	{
+		return Deinterlacer->Deinterlace(InVideoFrame);
+	}
+	return {};
+}
+
 TArray<ITimedDataInputChannel*> FMediaIOCorePlayerBase::GetChannels() const
 {
 	return Channels;
diff --git a/Engine/Plugins/Media/MediaIOFramework/Source/MediaIOCore/Private/Shared/MediaIOCoreTextureSampleBase.cpp b/Engine/Plugins/Media/MediaIOFramework/Source/MediaIOCore/Private/Shared/MediaIOCoreTextureSampleBase.cpp
--- a/Engine/Plugins/Media/MediaIOFramework/Source/MediaIOCore/Private/Shared/MediaIOCoreTextureSampleBase.cpp
+++ b/Engine/Plugins/Media/MediaIOFramework/Source/MediaIOCore/Private/Shared/MediaIOCoreTextureSampleBase.cpp
@@ -30,7 +30,20 @@ bool FMediaIOCoreTextureSampleBase::Initialize(const void* InVideoBuffer, uint32
 }
 
 
-bool FMediaIOCoreTextureSampleBase::Initialize(TArray<uint8> InVideoBuffer, uint32 InStride, uint32 InWidth, uint32 InHeight, EMediaTextureSampleFormat InSampleFormat, FTimespan InTime, const FFrameRate& InFrameRate, const TOptional<FTimecode>& InTimecode, bool bInIsSRGBInput)
+bool FMediaIOCoreTextureSampleBase::Initialize(const TArray<uint8>& InVideoBuffer, uint32 InStride, uint32 InWidth, uint32 InHeight, EMediaTextureSampleFormat InSampleFormat, FTimespan InTime, const FFrameRate& InFrameRate, const TOptional<FTimecode>& InTimecode, bool bInIsSRGBInput)
+{
+	FreeSample();
+
+	if (!SetProperties(InStride, InWidth, InHeight, InSampleFormat, InTime, InFrameRate, InTimecode, bInIsSRGBInput))
+	{
+		return false;
+	}
+
+	return SetBuffer(InVideoBuffer);
+}
+
+
+bool FMediaIOCoreTextureSampleBase::Initialize(TArray<uint8>&& InVideoBuffer, uint32 InStride, uint32 InWidth, uint32 InHeight, EMediaTextureSampleFormat InSampleFormat, FTimespan InTime, const FFrameRate& InFrameRate, const TOptional<FTimecode>& InTimecode, bool bInIsSRGBInput)
 {
 	FreeSample();
 
@@ -57,18 +70,29 @@ bool FMediaIOCoreTextureSampleBase::SetBuffer(const void* InVideoBuffer, uint32
 }
 
 
-bool FMediaIOCoreTextureSampleBase::SetBuffer(TArray<uint8> InVideoBuffer)
+bool FMediaIOCoreTextureSampleBase::SetBuffer(const TArray<uint8>& InVideoBuffer)
 {
 	if (InVideoBuffer.Num() == 0)
 	{
 		return false;
 	}
 
-	Buffer = MoveTemp(InVideoBuffer);
+	Buffer = InVideoBuffer;
 
 	return true;
 }
 
+bool FMediaIOCoreTextureSampleBase::SetBuffer(TArray<uint8>&& InVideoBuffer)
+{
+	if (InVideoBuffer.Num() == 0)
+	{
+		return false;
+	}
+
+	Buffer = MoveTemp(InVideoBuffer);
+
+	return true;
+}
 
 bool FMediaIOCoreTextureSampleBase::SetProperties(uint32 InStride, uint32 InWidth, uint32 InHeight, EMediaTextureSampleFormat InSampleFormat, FTimespan InTime, const FFrameRate& InFrameRate, const TOptional<FTimecode>& InTimecode, bool bInIsSRGBInput)
 {
@@ -89,7 +113,6 @@ bool FMediaIOCoreTextureSampleBase::SetProperties(uint32 InStride, uint32 InWidt
 	return true;
 }
 
-
 bool FMediaIOCoreTextureSampleBase::InitializeWithEvenOddLine(bool bUseEvenLine, const void* InVideoBuffer, uint32 InBufferSize, uint32 InStride, uint32 InWidth, uint32 InHeight, EMediaTextureSampleFormat InSampleFormat, FTimespan InTime, const FFrameRate& InFrameRate, const TOptional<FTimecode>& InTimecode, bool bInIsSRGBInput)
 {
 	FreeSample();
@@ -102,7 +125,6 @@ bool FMediaIOCoreTextureSampleBase::InitializeWithEvenOddLine(bool bUseEvenLine,
 	return SetBufferWithEvenOddLine(bUseEvenLine, InVideoBuffer, InBufferSize, InStride, InHeight);
 }
 
-
 bool FMediaIOCoreTextureSampleBase::SetBufferWithEvenOddLine(bool bUseEvenLine, const void* InVideoBuffer, uint32 InBufferSize, uint32 InStride, uint32 InHeight)
 {
 	Buffer.Reset(InBufferSize / 2);
@@ -116,7 +138,6 @@ bool FMediaIOCoreTextureSampleBase::SetBufferWithEvenOddLine(bool bUseEvenLine,
 	return true;
 }
 
-
 void* FMediaIOCoreTextureSampleBase::RequestBuffer(uint32 InBufferSize)
 {
 	FreeSample();
diff --git a/Engine/Plugins/Media/MediaIOFramework/Source/MediaIOCore/Private/Tests/MediaIODeinterlacerTests.cpp b/Engine/Plugins/Media/MediaIOFramework/Source/MediaIOCore/Private/Tests/MediaIODeinterlacerTests.cpp
new file mode 100644
--- /dev/null
+++ b/Engine/Plugins/Media/MediaIOFramework/Source/MediaIOCore/Private/Tests/MediaIODeinterlacerTests.cpp
@@ -0,0 +1,200 @@
+// Copyright Epic Games, Inc. All Rights Reserved.
+
+#include "CoreMinimal.h"
+#include "Misc/AutomationTest.h"
+#include "MediaIOCoreEncodeTime.h"
+#include "MediaIOCoreDeinterlacer.h"
+
+namespace UE::MediaIOTests
+{
+	constexpr uint8 FIRST_LINE_VALUE = 5;
+	constexpr uint8 SECOND_LINE_VALUE = 10;
+
+	TArray<uint8> CreateBuffer()
+	{
+		TArray<uint8> Buffer;
+		Buffer.Reserve(1080 * 1920);
+		
+		for (uint32 YIndex = 0; YIndex < 1080; YIndex++)
+		{
+			for (uint32 XIndex = 0; XIndex < 1920; XIndex++)
+			{
+				uint8 Value = FIRST_LINE_VALUE;
+				if ((YIndex % 2) != 0)
+				{
+					Value = SECOND_LINE_VALUE;
+				}
+
+				for (uint32 RGBAIndex = 0; RGBAIndex < 4; RGBAIndex++)
+				{
+					Buffer.Add(Value);
+				}
+			}
+		}
+
+		return Buffer;
+	}
+	
+	void TestBobDeinterlace(FAutomationTestBase& Test)
+	{
+		// Setup test data
+		TArray<uint8> Buffer = CreateBuffer();
+
+		EMediaTextureSampleFormat VideoSampleFormat = EMediaTextureSampleFormat::CharBGRA;
+
+		constexpr uint32 Width = 1920;
+		constexpr uint32 Height = 1080;
+		constexpr uint32 Stride = Width * 4;
+
+		UE::MediaIOCore::FVideoFrame FrameInfo
+		{
+			Buffer.GetData(),
+			(uint32)Buffer.Num(),
+			Stride,
+			Width,
+			Height,
+			EMediaTextureSampleFormat::CharBGRA,
+			FTimespan::FromSeconds(0),
+			FFrameRate(),
+			FTimecode(),
+			UE::Color::EColorSpace::None
+		};
+
+		// Execute test
+		UE::MediaIOCore::FBobDeinterlacer Deinterlacer(UE::MediaIOCore::FDeinterlacer::FOnAcquireSample_AnyThread::CreateLambda([](){ return MakeShared<FMediaIOCoreTextureSampleBase>(); }), EMediaIOInterlaceFieldOrder::TopFieldFirst);
+		TArray<TSharedRef<FMediaIOCoreTextureSampleBase>> DeinterlacedSamples = Deinterlacer.Deinterlace(FrameInfo);
+
+		// Validate result
+		Test.TestEqual(TEXT("Samples were created"), DeinterlacedSamples.Num(), 2);
+		uint8* Buffer1 = (uint8*)DeinterlacedSamples[0]->GetBuffer();
+		uint8* Buffer2 = (uint8*)DeinterlacedSamples[1]->GetBuffer();
+
+		Test.TestEqual(TEXT("First sample first line is correct"), *Buffer1, FIRST_LINE_VALUE);
+		Test.TestEqual(TEXT("First sample Second line is correct"), *(Buffer1 + Stride), FIRST_LINE_VALUE);
+
+		Test.TestEqual(TEXT("Second sample first line is correct"), *Buffer2, SECOND_LINE_VALUE);
+		Test.TestEqual(TEXT("Second sample Second line is correct"), *(Buffer2 + Stride), SECOND_LINE_VALUE);
+	}
+
+	void TestBlendDeinterlace(FAutomationTestBase& Test)
+	{
+		// Setup test data
+		TArray<uint8> Buffer = CreateBuffer();
+
+		EMediaTextureSampleFormat VideoSampleFormat = EMediaTextureSampleFormat::CharBGRA;
+
+		constexpr uint32 Width = 1920;
+		constexpr uint32 Height = 1080;
+		constexpr uint32 Stride = Width * 4;
+
+		UE::MediaIOCore::FVideoFrame FrameInfo
+		{
+			Buffer.GetData(),
+			(uint32)Buffer.Num(),
+			Stride,
+			Width,
+			Height,
+			EMediaTextureSampleFormat::CharBGRA,
+			FTimespan::FromSeconds(0),
+			FFrameRate(),
+			FTimecode(),
+			UE::Color::EColorSpace::None
+		};
+
+		// Execute test
+		UE::MediaIOCore::FBlendDeinterlacer Deinterlacer(UE::MediaIOCore::FDeinterlacer::FOnAcquireSample_AnyThread::CreateLambda([]() { return MakeShared<FMediaIOCoreTextureSampleBase>(); }), EMediaIOInterlaceFieldOrder::TopFieldFirst);
+		TArray<TSharedRef<FMediaIOCoreTextureSampleBase>> DeinterlacedSamples = Deinterlacer.Deinterlace(FrameInfo);
+
+		// Validate result
+		Test.TestEqual(TEXT("One Sample was created"), DeinterlacedSamples.Num(), 1);
+		uint8* Buffer1 = (uint8*)DeinterlacedSamples[0]->GetBuffer();
+
+		Test.TestEqual(TEXT("Sample line is averaged"), *Buffer1, FMath::DivideAndRoundNearest(FIRST_LINE_VALUE + SECOND_LINE_VALUE, 2));
+	}
+
+	void TestDiscardDeinterlace(FAutomationTestBase& Test)
+	{
+		// Setup test data
+		TArray<uint8> Buffer = CreateBuffer();
+
+		EMediaTextureSampleFormat VideoSampleFormat = EMediaTextureSampleFormat::CharBGRA;
+
+		constexpr uint32 Width = 1920;
+		constexpr uint32 Height = 1080;
+		constexpr uint32 Stride = Width * 4;
+
+		UE::MediaIOCore::FVideoFrame FrameInfo
+		{
+			Buffer.GetData(),
+			(uint32)Buffer.Num(),
+			Stride,
+			Width,
+			Height,
+			EMediaTextureSampleFormat::CharBGRA,
+			FTimespan::FromSeconds(0),
+			FFrameRate(),
+			FTimecode(),
+			UE::Color::EColorSpace::None
+		};
+
+		// Execute test
+		UE::MediaIOCore::FDiscardDeinterlacer Deinterlacer(UE::MediaIOCore::FDeinterlacer::FOnAcquireSample_AnyThread::CreateLambda([]() { return MakeShared<FMediaIOCoreTextureSampleBase>(); }), EMediaIOInterlaceFieldOrder::TopFieldFirst);
+		TArray<TSharedRef<FMediaIOCoreTextureSampleBase>> DeinterlacedSamples = Deinterlacer.Deinterlace(FrameInfo);
+
+		// Validate result
+		Test.TestEqual(TEXT("One sample was created"), DeinterlacedSamples.Num(), 1);
+		uint8* Buffer1 = (uint8*)DeinterlacedSamples[0]->GetBuffer();
+
+		Test.TestEqual(TEXT("First sample first line is correct"), *Buffer1, FIRST_LINE_VALUE);
+		Test.TestEqual(TEXT("First sample Second line is correct"), *(Buffer1 + Stride), FIRST_LINE_VALUE);
+	}
+
+	void TestDiscardDeinterlaceBottomFieldOrder(FAutomationTestBase& Test)
+	{
+		// Setup test data
+		TArray<uint8> Buffer = CreateBuffer();
+
+		EMediaTextureSampleFormat VideoSampleFormat = EMediaTextureSampleFormat::CharBGRA;
+
+		constexpr uint32 Width = 1920;
+		constexpr uint32 Height = 1080;
+		constexpr uint32 Stride = Width * 4;
+
+		UE::MediaIOCore::FVideoFrame FrameInfo
+		{
+			Buffer.GetData(),
+			(uint32)Buffer.Num(),
+			Stride,
+			Width,
+			Height,
+			EMediaTextureSampleFormat::CharBGRA,
+			FTimespan::FromSeconds(0),
+			FFrameRate(),
+			FTimecode(),
+			UE::Color::EColorSpace::None
+		};
+
+		// Execute test
+		UE::MediaIOCore::FDiscardDeinterlacer Deinterlacer(UE::MediaIOCore::FDeinterlacer::FOnAcquireSample_AnyThread::CreateLambda([]() { return MakeShared<FMediaIOCoreTextureSampleBase>(); }), EMediaIOInterlaceFieldOrder::BottomFieldFirst);
+		TArray<TSharedRef<FMediaIOCoreTextureSampleBase>> DeinterlacedSamples = Deinterlacer.Deinterlace(FrameInfo);
+
+		// Validate result
+		Test.TestEqual(TEXT("One sample was created"), DeinterlacedSamples.Num(), 1);
+		uint8* Buffer1 = (uint8*)DeinterlacedSamples[0]->GetBuffer();
+
+		Test.TestEqual(TEXT("First sample first line is correct"), *Buffer1, SECOND_LINE_VALUE);
+		Test.TestEqual(TEXT("First sample Second line is correct"), *(Buffer1 + Stride), SECOND_LINE_VALUE);
+	}
+}
+
+
+IMPLEMENT_SIMPLE_AUTOMATION_TEST(FMediaIODeinterlacerTests, "Plugin.MediaIO.Deinterlace", EAutomationTestFlags::EditorContext | EAutomationTestFlags::EngineFilter)
+bool FMediaIODeinterlacerTests::RunTest(const FString& Parameters)
+{
+	// Test expose whole container
+	UE::MediaIOTests::TestBobDeinterlace(*this);
+	UE::MediaIOTests::TestBlendDeinterlace(*this);
+	UE::MediaIOTests::TestDiscardDeinterlace(*this);
+	UE::MediaIOTests::TestDiscardDeinterlaceBottomFieldOrder(*this);
+	return true;
+}
\ No newline at end of file
diff --git a/Engine/Plugins/Media/MediaIOFramework/Source/MediaIOCore/Public/CaptureCardMediaSource.h b/Engine/Plugins/Media/MediaIOFramework/Source/MediaIOCore/Public/CaptureCardMediaSource.h
new file mode 100644
--- /dev/null
+++ b/Engine/Plugins/Media/MediaIOFramework/Source/MediaIOCore/Public/CaptureCardMediaSource.h
@@ -0,0 +1,56 @@
+// Copyright Epic Games, Inc. All Rights Reserved.
+
+#pragma once
+
+#include "TimeSynchronizableMediaSource.h"
+#include "Containers/UnrealString.h"
+#include "MediaIOCoreDefinitions.h"
+#include "MediaIOCoreDeinterlacer.h"
+#include "HAL/Platform.h"
+#include "UObject/NameTypes.h"
+#include "UObject/SoftObjectPath.h"
+#include "UObject/ObjectMacros.h"
+#include "UObject/UObjectGlobals.h"
+
+#include "CaptureCardMediaSource.generated.h"
+
+namespace UE::CaptureCardMediaSource
+{
+	static FName Deinterlacer("Deinterlacer");
+	static FName InterlaceFieldOrder("InterlaceFieldOrder");
+}
+
+/**
+ * Base class for media sources that are coming from a capture card.
+ */
+UCLASS(Abstract)
+class MEDIAIOCORE_API UCaptureCardMediaSource : public UTimeSynchronizableMediaSource
+{
+	GENERATED_BODY()
+	
+public:
+	
+	UCaptureCardMediaSource()
+	{
+		Deinterlacer = CreateDefaultSubobject<UBobDeinterlacer>("Deinterlacer");
+	}
+	
+	/**
+	 * How interlaced video should be treated.
+	 */
+	UPROPERTY(BlueprintReadOnly, Instanced, EditAnywhere, Category = "Video")
+	TObjectPtr<UVideoDeinterlacer> Deinterlacer;
+
+	/**
+	 * The order in which interlace fields should be copied.
+	 */
+	UPROPERTY(BlueprintReadOnly, EditAnywhere, Category = "Video")
+	EMediaIOInterlaceFieldOrder InterlaceFieldOrder = EMediaIOInterlaceFieldOrder::TopFieldFirst;
+
+public:
+	//~ IMediaOptions interface
+	using Super::GetMediaOption;
+	virtual FString GetMediaOption(const FName& Key, const FString& DefaultValue) const override;
+	virtual int64 GetMediaOption(const FName& Key, int64 DefaultValue) const override;
+	virtual bool HasMediaOption(const FName& Key) const override;
+};
diff --git a/Engine/Plugins/Media/MediaIOFramework/Source/MediaIOCore/Public/MediaIOCoreDefinitions.h b/Engine/Plugins/Media/MediaIOFramework/Source/MediaIOCore/Public/MediaIOCoreDefinitions.h
--- a/Engine/Plugins/Media/MediaIOFramework/Source/MediaIOCore/Public/MediaIOCoreDefinitions.h
+++ b/Engine/Plugins/Media/MediaIOFramework/Source/MediaIOCore/Public/MediaIOCoreDefinitions.h
@@ -66,6 +66,16 @@ enum class EMediaIOAutoDetectableTimecodeFormat : uint8
 	Auto = 255,
 };
 
+/**
+ * Which interlace field should be set first in the buffer.
+ */
+UENUM()
+enum class EMediaIOInterlaceFieldOrder : uint8
+{
+	TopFieldFirst,
+	BottomFieldFirst
+};
+
 namespace UE::MediaIO
 {
 	static EMediaIOAutoDetectableTimecodeFormat ToAutoDetectableTimecodeFormat(EMediaIOTimecodeFormat TimecodeFormat)
@@ -388,4 +398,4 @@ public:
 	 * @return String representation, i.e. "Video/Single1/1080p30fps/LTC".
 	 */
 	FText ToText(bool bAutoDetected = false) const;
-};
\ No newline at end of file
+};
diff --git a/Engine/Plugins/Media/MediaIOFramework/Source/MediaIOCore/Public/MediaIOCoreDeinterlacer.h b/Engine/Plugins/Media/MediaIOFramework/Source/MediaIOCore/Public/MediaIOCoreDeinterlacer.h
new file mode 100644
--- /dev/null
+++ b/Engine/Plugins/Media/MediaIOFramework/Source/MediaIOCore/Public/MediaIOCoreDeinterlacer.h
@@ -0,0 +1,167 @@
+// Copyright Epic Games, Inc. All Rights Reserved.
+#pragma once
+
+#include "ColorSpace.h"
+#include "Containers/Array.h"
+#include "MediaIOCoreDefinitions.h"
+#include "MediaIOCoreTextureSampleBase.h"
+#include "Templates/SharedPointer.h"
+
+#include "MediaIOCoreDeinterlacer.generated.h"
+
+namespace UE::MediaIOCore
+{
+	/** Description of a video frame. */
+	struct MEDIAIOCORE_API FVideoFrame
+	{
+		FVideoFrame(const void* InVideoBuffer
+		, uint32 InBufferSize
+		, uint32 InStride
+		, uint32 InWidth
+		, uint32 InHeight
+		, EMediaTextureSampleFormat InSampleFormat
+		, FTimespan InTime
+		, const FFrameRate& InFrameRate
+		, const TOptional<FTimecode>& InTimecode
+		, Color::EColorSpace InColorSpace)
+			: VideoBuffer(InVideoBuffer)
+			, BufferSize(InBufferSize)
+			, Stride(InStride)
+			, Width(InWidth)
+			, Height(InHeight)
+			, SampleFormat(InSampleFormat)
+			, Time(InTime)
+			, FrameRate(InFrameRate)
+			, Timecode(InTimecode)
+			, ColorSpace(InColorSpace)
+		{
+		}
+
+		const void* VideoBuffer;
+		uint32 BufferSize;
+		uint32 Stride;
+		uint32 Width;
+		uint32 Height;
+		EMediaTextureSampleFormat SampleFormat;
+		FTimespan Time;
+		const FFrameRate& FrameRate;
+		const TOptional<FTimecode>& Timecode;
+		Color::EColorSpace ColorSpace;
+	};
+
+	/**
+	 * Handles deinterlacing a video signal.
+	 */
+	class MEDIAIOCORE_API FDeinterlacer
+	{
+	public:
+		DECLARE_DELEGATE_RetVal(TSharedPtr<FMediaIOCoreTextureSampleBase>, FOnAcquireSample_AnyThread);
+
+		FDeinterlacer() = default;
+		FDeinterlacer(FOnAcquireSample_AnyThread InOnAcquireSample, EMediaIOInterlaceFieldOrder InInterlaceFieldOrder) 
+			: AcquireSampleDelegate(MoveTemp(InOnAcquireSample))
+			, InterlaceFieldOrder(InInterlaceFieldOrder)
+		{
+		}
+
+		virtual ~FDeinterlacer() {}
+        
+		/**
+		 * Default implementation which applies no deinterlacing.
+		 * @param InVideoFrame The video buffer and its metadata.
+		 * @return One or more sample depending on the deinterlacing method.
+		 */
+		virtual TArray<TSharedRef<FMediaIOCoreTextureSampleBase>> Deinterlace(const FVideoFrame& InVideoFrame) const;
+	protected:
+		/** Delegate called to acquire samples. */
+		FOnAcquireSample_AnyThread AcquireSampleDelegate;
+		/** What order should fields be read in. */
+		EMediaIOInterlaceFieldOrder InterlaceFieldOrder = EMediaIOInterlaceFieldOrder::TopFieldFirst; 
+	};
+	
+	/** Double field lines, keeping temporal resolution but halving texture vertical resolution. */
+	class FBobDeinterlacer : public FDeinterlacer
+	{
+	public:
+		FBobDeinterlacer(FOnAcquireSample_AnyThread InOnAcquireSample, EMediaIOInterlaceFieldOrder InInterlaceFieldOrder) 
+			: FDeinterlacer(MoveTemp(InOnAcquireSample), InInterlaceFieldOrder)
+		{
+		}
+		virtual TArray<TSharedRef<FMediaIOCoreTextureSampleBase>> Deinterlace(const FVideoFrame& InVideoFrame) const override;
+	};
+
+	/** Blend the top and bottom field, halving temporal and texture vertical resoluton. */
+	class FBlendDeinterlacer : public FDeinterlacer
+	{
+	public:
+		FBlendDeinterlacer(FOnAcquireSample_AnyThread InOnAcquireSample, EMediaIOInterlaceFieldOrder InInterlaceFieldOrder) 
+			: FDeinterlacer(MoveTemp(InOnAcquireSample), InInterlaceFieldOrder)
+		{
+		}
+		virtual TArray<TSharedRef<FMediaIOCoreTextureSampleBase>> Deinterlace(const FVideoFrame& InVideoFrame) const override;
+	};
+
+	/** Discards one of the field entirely (based on InterlaceFieldOrder), halving both temporal and spatial resolutions. */
+	class FDiscardDeinterlacer : public FDeinterlacer
+	{
+	public:
+		FDiscardDeinterlacer(FOnAcquireSample_AnyThread InOnAcquireSample, EMediaIOInterlaceFieldOrder InInterlaceFieldOrder) 
+			: FDeinterlacer(MoveTemp(InOnAcquireSample), InInterlaceFieldOrder)
+		{
+		}
+		virtual TArray<TSharedRef<FMediaIOCoreTextureSampleBase>> Deinterlace(const FVideoFrame& InVideoFrame) const override;
+	};
+	
+	
+}
+
+/**
+ * Represents a deinterlacing algorithm. Will dictate how the incoming video signal is converted from interlace to a progressive signal.
+ */
+UCLASS(Abstract, editinlinenew)
+class MEDIAIOCORE_API UVideoDeinterlacer : public UObject
+{
+	GENERATED_BODY()
+public:
+	/**
+	 * Creates an instance of a FDeinterlacer that represents a deinterlacing algorithm.
+	 */
+	virtual TSharedPtr<UE::MediaIOCore::FDeinterlacer> Instantiate(UE::MediaIOCore::FDeinterlacer::FOnAcquireSample_AnyThread InAcquireSampleDelegate, EMediaIOInterlaceFieldOrder InInterlaceFieldOrder) const PURE_VIRTUAL(UVideoDeinterlacer::Instantiate, return nullptr; );
+};
+
+
+/** Double field lines, keeping temporal resolution but halving texture vertical resolution. */
+UCLASS()
+class MEDIAIOCORE_API UBobDeinterlacer : public UVideoDeinterlacer
+{
+	GENERATED_BODY()
+public:
+	virtual TSharedPtr<UE::MediaIOCore::FDeinterlacer> Instantiate(UE::MediaIOCore::FDeinterlacer::FOnAcquireSample_AnyThread InAcquireSampleDelegate, EMediaIOInterlaceFieldOrder InInterlaceFieldOrder) const override
+	{
+		return MakeShared<UE::MediaIOCore::FBobDeinterlacer>(MoveTemp(InAcquireSampleDelegate), InInterlaceFieldOrder);
+	}
+};
+
+/** Blend the top and bottom field, keeping temporal resolution but halving texture vertical resoluton. */
+UCLASS()
+class MEDIAIOCORE_API UBlendDeinterlacer : public UVideoDeinterlacer
+{
+	GENERATED_BODY()
+public:
+	virtual TSharedPtr<UE::MediaIOCore::FDeinterlacer> Instantiate(UE::MediaIOCore::FDeinterlacer::FOnAcquireSample_AnyThread InAcquireSampleDelegate, EMediaIOInterlaceFieldOrder InInterlaceFieldOrder) const override
+	{
+		return MakeShared<UE::MediaIOCore::FBlendDeinterlacer>(MoveTemp(InAcquireSampleDelegate), InInterlaceFieldOrder);
+	}
+};
+
+/** Discards one of the field entirely (based on InterlaceFieldOrder), halving both temporal and spatial resolutions. */
+UCLASS()
+class MEDIAIOCORE_API UDiscardDeinterlacer : public UVideoDeinterlacer
+{
+	GENERATED_BODY()
+public:
+	virtual TSharedPtr<UE::MediaIOCore::FDeinterlacer> Instantiate(UE::MediaIOCore::FDeinterlacer::FOnAcquireSample_AnyThread InAcquireSampleDelegate, EMediaIOInterlaceFieldOrder InInterlaceFieldOrder) const override
+	{
+		return MakeShared<UE::MediaIOCore::FDiscardDeinterlacer>(MoveTemp(InAcquireSampleDelegate), InInterlaceFieldOrder);
+	}
+};
diff --git a/Engine/Plugins/Media/MediaIOFramework/Source/MediaIOCore/Public/MediaIOCorePlayerBase.h b/Engine/Plugins/Media/MediaIOFramework/Source/MediaIOCore/Public/MediaIOCorePlayerBase.h
--- a/Engine/Plugins/Media/MediaIOFramework/Source/MediaIOCore/Public/MediaIOCorePlayerBase.h
+++ b/Engine/Plugins/Media/MediaIOFramework/Source/MediaIOCore/Public/MediaIOCorePlayerBase.h
@@ -42,6 +42,12 @@ namespace UE::GPUTextureTransfer
 {
 	using TextureTransferPtr = TSharedPtr<class ITextureTransfer>;
 }
+
+namespace UE::MediaIOCore
+{
+	class FDeinterlacer;
+	struct FVideoFrame;
+}
  
 /**
  * Implements a base player for hardware IO cards. 
@@ -188,6 +194,13 @@ protected:
 
 	virtual bool CanUseGPUTextureTransfer();
 
+	/** Handler used to provide texture samples to the deinterlacer. */
+	virtual TSharedPtr<FMediaIOCoreTextureSampleBase> AcquireSample_AnyThread() const
+	{
+		ensureMsgf(false, TEXT("If deinterlacing is used by the media player, AcquireSample_AnyThread must be implemented by the child media player class."));
+		return nullptr;
+	}
+
 	void OnSampleDestroyed(TRefCountPtr<FRHITexture> InTexture);
 	void RegisterSampleBuffer(const TSharedPtr<FMediaIOCoreTextureSampleBase>& InSample);
 	void UnregisterSampleBuffers();
@@ -196,6 +209,9 @@ protected:
 	void PreGPUTransfer(const TSharedPtr<FMediaIOCoreTextureSampleBase>& InSample);
 	void ExecuteGPUTransfer(const TSharedPtr<FMediaIOCoreTextureSampleBase>& InSample);
 
+	/** Deinterlace a video frame. Only valid to call when the video stream is open. */
+	TArray<TSharedRef<FMediaIOCoreTextureSampleBase>> Deinterlace(const UE::MediaIOCore::FVideoFrame& InVideoFrame) const;
+
 protected:
 	/** Critical section for synchronizing access to receiver and sinks. */
 	FCriticalSection CriticalSection;
@@ -260,5 +276,7 @@ private:
 	TSet<void*> RegisteredBuffers;
 	/** Pool of textures registerd with GPU Texture transfer. */
 	TSet<TRefCountPtr<FRHITexture>> RegisteredTextures;
+	/** Utility object to handle deinterlacing. */
+	TSharedPtr<UE::MediaIOCore::FDeinterlacer> Deinterlacer;
 
 };
diff --git a/Engine/Plugins/Media/MediaIOFramework/Source/MediaIOCore/Public/MediaIOCoreTextureSampleBase.h b/Engine/Plugins/Media/MediaIOFramework/Source/MediaIOCore/Public/MediaIOCoreTextureSampleBase.h
--- a/Engine/Plugins/Media/MediaIOFramework/Source/MediaIOCore/Public/MediaIOCoreTextureSampleBase.h
+++ b/Engine/Plugins/Media/MediaIOFramework/Source/MediaIOCore/Public/MediaIOCoreTextureSampleBase.h
@@ -3,6 +3,7 @@
 #pragma once
 
 #include "IMediaTextureSample.h"
+#include "MediaIOCoreDefinitions.h"
 #include "MediaObjectPool.h"
 #include "Misc/FrameRate.h"
 #include "Templates/RefCounting.h"
@@ -49,7 +50,22 @@ public:
 	 * @param InTimecode The sample timecode if available.
 	 * @param bInIsSRGBInput The sample source is in sSRGB so apply a linear conversion.
 	 */
-	bool Initialize(TArray<uint8> InVideoBuffer, uint32 InStride, uint32 InWidth, uint32 InHeight, EMediaTextureSampleFormat InSampleFormat, FTimespan InTime, const FFrameRate& InFrameRate, const TOptional<FTimecode>& InTimecode, bool bInIsSRGBInput);
+	bool Initialize(const TArray<uint8>& InVideoBuffer, uint32 InStride, uint32 InWidth, uint32 InHeight, EMediaTextureSampleFormat InSampleFormat, FTimespan InTime, const FFrameRate& InFrameRate, const TOptional<FTimecode>& InTimecode, bool bInIsSRGBInput);
+
+	/**
+	 * Initialize the sample.
+	 *
+	 * @param InVideoBuffer The video frame data.
+	 * @param InStride The number of channel of the video buffer.
+	 * @param InWidth The sample rate of the video buffer.
+	 * @param InHeight The sample rate of the video buffer.
+	 * @param InSampleFormat The sample format of the video buffer.
+	 * @param InTime The sample time (in the player's own clock).
+	 * @param InFrameRate The framerate of the media that produce the sample.
+	 * @param InTimecode The sample timecode if available.
+	 * @param bInIsSRGBInput The sample source is in sSRGB so apply a linear conversion.
+	 */
+	bool Initialize(TArray<uint8>&& InVideoBuffer, uint32 InStride, uint32 InWidth, uint32 InHeight, EMediaTextureSampleFormat InSampleFormat, FTimespan InTime, const FFrameRate& InFrameRate, const TOptional<FTimecode>& InTimecode, bool bInIsSRGBInput);
 
 	/**
 	 * Initialize the sample.
@@ -64,7 +80,14 @@ public:
 	 *
 	 * @param InVideoBuffer The video frame data.
 	 */
-	bool SetBuffer(TArray<uint8> InVideoBuffer);
+	bool SetBuffer(const TArray<uint8>& InVideoBuffer);
+
+	/**
+	 * Set the sample buffer.
+	 *
+	 * @param InVideoBuffer The video frame data.
+	 */
+	bool SetBuffer(TArray<uint8>&& InVideoBuffer);
 
 	/**
 	 * Set the sample properties.
@@ -94,6 +117,7 @@ public:
 	 * @param bInIsSRGBInput The sample source is in sSRGB so apply a linear conversion.
 	 */
 	bool InitializeWithEvenOddLine(bool bUseEvenLine, const void* InVideoBuffer, uint32 InBufferSize, uint32 InStride, uint32 InWidth, uint32 InHeight, EMediaTextureSampleFormat InSampleFormat, FTimespan InTime, const FFrameRate& InFrameRate, const TOptional<FTimecode>& InTimecode, bool bInIsSRGBInput);
+
 	/**
 	 * Set the sample buffer with half it's original height and take only the odd or even line.
 	 *
@@ -105,6 +129,8 @@ public:
 	 */
 	bool SetBufferWithEvenOddLine(bool bUseEvenLine, const void* InVideoBuffer, uint32 InBufferSize, uint32 InStride, uint32 InHeight);
 
+
+
 	/**
 	 * Request an uninitialized sample buffer.
 	 * Should be used when the buffer could be filled by something else.
