Fix OCIO being applied twice when doing media capture
Fix crash when outputting with media capture in different capture phase than end frame

diff --git a/Engine/Plugins/Media/MediaIOFramework/Source/MediaIOCore/Private/MediaCapture.cpp b/Engine/Plugins/Media/MediaIOFramework/Source/MediaIOCore/Private/MediaCapture.cpp
--- a/Engine/Plugins/Media/MediaIOFramework/Source/MediaIOCore/Private/MediaCapture.cpp
+++ b/Engine/Plugins/Media/MediaIOFramework/Source/MediaIOCore/Private/MediaCapture.cpp
@@ -19,6 +19,7 @@
 #include "MediaShaders.h"
 #include "Misc/App.h"
 #include "Misc/ScopeLock.h"
+#include "OpenColorIODisplayExtension.h"
 #include "RendererInterface.h"
 #include "RenderGraphResources.h"
 #include "RenderGraphUtils.h"
@@ -290,7 +291,9 @@ bool UMediaCapture::StartSourceCapture(TSharedPtr<UE::MediaCapture::Private::FCa
 		
 		if (DesiredCaptureOptions.CapturePhase != EMediaCapturePhase::EndFrame)
 		{
-			ViewExtension = FSceneViewExtensions::NewExtension<FMediaCaptureSceneViewExtension>(this, DesiredCaptureOptions.CapturePhase);
+			/** If we are running with color conversion, we need our extension to be executed before the OCIO extension is applied to the viewport. */
+			const int32 Priority =  DesiredCaptureOptions.ColorConversionSettings.IsValid() ? OPENCOLORIO_SCENE_VIEW_EXTENSION_PRIORITY + 1 : 0;
+			ViewExtension = FSceneViewExtensions::NewExtension<FMediaCaptureSceneViewExtension>(this, DesiredCaptureOptions.CapturePhase, Priority);
 		}
 
 		if (ViewExtension)
@@ -663,8 +666,11 @@ void UMediaCapture::CaptureImmediate_RenderThread(const UE::MediaCaptureData::FC
 		return;
 	}
 
-	// Keep resource size up to date with incoming resource to capture
- 	StaticCastSharedPtr<UE::MediaCapture::Private::FRHIResourceCaptureSource>(CaptureSource)->ResourceDescription.ResourceSize = Args.GetSizeXY();
+	if (CaptureSource->GetSourceType() == EMediaCaptureSourceType::RHI_RESOURCE)
+	{
+		// Keep resource size up to date with incoming resource to capture
+ 		StaticCastSharedPtr<UE::MediaCapture::Private::FRHIResourceCaptureSource>(CaptureSource)->ResourceDescription.ResourceSize = Args.GetSizeXY();
+	}
 
 	// Get cached capture data from game thread. We want to find a cached frame matching current render thread frame number
 	bool bFoundMatchingData = false;
diff --git a/Engine/Plugins/Media/MediaIOFramework/Source/MediaIOCore/Private/MediaCaptureSceneViewExtension.h b/Engine/Plugins/Media/MediaIOFramework/Source/MediaIOCore/Private/MediaCaptureSceneViewExtension.h
--- a/Engine/Plugins/Media/MediaIOFramework/Source/MediaIOCore/Private/MediaCaptureSceneViewExtension.h
+++ b/Engine/Plugins/Media/MediaIOFramework/Source/MediaIOCore/Private/MediaCaptureSceneViewExtension.h
@@ -24,14 +24,21 @@ class FMediaCaptureSceneViewExtension : public FSceneViewExtensionBase
 {
 public:
 	
-	FMediaCaptureSceneViewExtension(const FAutoRegister& InAutoRegister, UMediaCapture* InMediaCapture, EMediaCapturePhase InCapturePhase)
+	FMediaCaptureSceneViewExtension(const FAutoRegister& InAutoRegister, UMediaCapture* InMediaCapture, EMediaCapturePhase InCapturePhase, int32 InPriority)
 		: FSceneViewExtensionBase(InAutoRegister)
 		, WeakCapture(InMediaCapture)
 		, CapturePhase(InCapturePhase)
+		, Priority(InPriority)
 	{
+		
 	}
 
 	//~ Begin FSceneViewExtensionBase Interface
+	virtual int32 GetPriority() const override
+	{
+		return Priority;
+	}
+
 	virtual void SetupViewFamily(FSceneViewFamily& InViewFamily) override {};
 	virtual void SetupView(FSceneViewFamily& InViewFamily, FSceneView& InView) override 
 	{
@@ -81,6 +88,7 @@ public:
 	{
 		return true;
 	}
+	//~ End FSceneViewExtensionBase Interface
 
 	FScreenPassTexture PostProcessCallback_RenderThread(FRDGBuilder& GraphBuilder, const FSceneView& View, const FPostProcessMaterialInputs& InOutInputs)
 	{
@@ -129,6 +137,6 @@ private:
 	bool bPostProcessingEnabled = true;
 	bool bValidPhase = true;
 	FString LastErrorMessage;
-	//~ End FSceneViewExtensionBase Interface
+	int32 Priority = 0;
 };
 
