Update bob deinterlace to double lines in the CPU instead

diff --git a/Engine/Plugins/Media/MediaIOFramework/Source/MediaIOCore/Private/MediaIOCoreDeinterlacer.cpp b/Engine/Plugins/Media/MediaIOFramework/Source/MediaIOCore/Private/MediaIOCoreDeinterlacer.cpp
--- a/Engine/Plugins/Media/MediaIOFramework/Source/MediaIOCore/Private/MediaIOCoreDeinterlacer.cpp
+++ b/Engine/Plugins/Media/MediaIOFramework/Source/MediaIOCore/Private/MediaIOCoreDeinterlacer.cpp
@@ -41,10 +41,28 @@ namespace UE::MediaIOCore
 
 		if (TextureSampleEven && TextureSampleOdd)
 		{
-			if (TextureSampleEven->InitializeWithEvenOddLine(
-				  InterlaceFieldOrder == EMediaIOInterlaceFieldOrder::TopFieldFirst
-				, InVideoFrame.VideoBuffer
-				, InVideoFrame.BufferSize
+			TArray<uint8> EvenBuffer;
+			EvenBuffer.Reserve(InVideoFrame.BufferSize);
+
+			TArray<uint8> OddBuffer;
+			OddBuffer.Reserve(InVideoFrame.BufferSize);
+
+			for (uint32 IndexY = (InterlaceFieldOrder == EMediaIOInterlaceFieldOrder::TopFieldFirst ? 0 : 1); IndexY < InVideoFrame.Height; IndexY += 2)
+			{
+				const uint8* Source = reinterpret_cast<const uint8*>(InVideoFrame.VideoBuffer) + (IndexY * InVideoFrame.Stride);
+				EvenBuffer.Append(Source, InVideoFrame.Stride);
+				EvenBuffer.Append(Source, InVideoFrame.Stride);
+			}
+			
+			for (uint32 IndexY = (InterlaceFieldOrder == EMediaIOInterlaceFieldOrder::TopFieldFirst ? 1 : 0); IndexY < InVideoFrame.Height; IndexY += 2)
+			{
+				const uint8* Source = reinterpret_cast<const uint8*>(InVideoFrame.VideoBuffer) + (IndexY * InVideoFrame.Stride);
+				OddBuffer.Append(Source, InVideoFrame.Stride);
+				OddBuffer.Append(Source, InVideoFrame.Stride);
+			}
+
+			if (TextureSampleEven->Initialize(
+				  MoveTemp(EvenBuffer)
 				, InVideoFrame.Stride
 				, InVideoFrame.Width
 				, InVideoFrame.Height
@@ -57,11 +75,10 @@ namespace UE::MediaIOCore
 				Samples.Add(TextureSampleEven.ToSharedRef());
 			}
 
+
 			// Don't create second sample if first one fails in order to avoid introducing field flipping.
-			if (Samples.Num() && TextureSampleOdd->InitializeWithEvenOddLine(
-				 !(InterlaceFieldOrder == EMediaIOInterlaceFieldOrder::TopFieldFirst)
-				, InVideoFrame.VideoBuffer
-				, InVideoFrame.BufferSize
+			if (Samples.Num() && TextureSampleOdd->Initialize(
+				  MoveTemp(OddBuffer)
 				, InVideoFrame.Stride
 				, InVideoFrame.Width
 				, InVideoFrame.Height
