Add YUV conversion matrix override

diff --git a/Engine/Plugins/Media/MediaIOFramework/Source/MediaIOCore/Private/MediaCapture.cpp b/Engine/Plugins/Media/MediaIOFramework/Source/MediaIOCore/Private/MediaCapture.cpp
--- a/Engine/Plugins/Media/MediaIOFramework/Source/MediaIOCore/Private/MediaCapture.cpp
+++ b/Engine/Plugins/Media/MediaIOFramework/Source/MediaIOCore/Private/MediaCapture.cpp
@@ -898,6 +898,11 @@ void UMediaCapture::ResetFixedViewportSize(TSharedPtr<FSceneViewport> InViewport
 	}
 }
 
+const FMatrix& UMediaCapture::GetRGBToYUVConversionMatrix() const
+{
+	return MediaShaders::RgbToYuvRec709Scaled;
+}
+
 void UMediaCapture::WaitForSingleExperimentalSchedulingTaskToComplete()
 {
 	if (UseExperimentalScheduling())
diff --git a/Engine/Plugins/Media/MediaIOFramework/Source/MediaIOCore/Private/MediaCaptureRenderPass.cpp b/Engine/Plugins/Media/MediaIOFramework/Source/MediaIOCore/Private/MediaCaptureRenderPass.cpp
--- a/Engine/Plugins/Media/MediaIOFramework/Source/MediaIOCore/Private/MediaCaptureRenderPass.cpp
+++ b/Engine/Plugins/Media/MediaIOFramework/Source/MediaIOCore/Private/MediaCaptureRenderPass.cpp
@@ -465,6 +465,8 @@ namespace UE::MediaCapture
 			FGlobalShaderMap* GlobalShaderMap = GetGlobalShaderMap(GMaxRHIFeatureLevel);
 			TShaderMapRef<FScreenPassVS> VertexShader(GlobalShaderMap);
 
+			const FMatrix& ConversionMatrix = Args.MediaCapture->GetRGBToYUVConversionMatrix();
+
 			switch (Args.MediaCapture->GetConversionOperation())
 			{
 			case EMediaCaptureConversionOperation::RGBA8_TO_YUV_8BIT:
@@ -473,7 +475,7 @@ namespace UE::MediaCapture
 					FScreenPassTextureViewport InputViewport(ConversionPassArgs.SourceRGBTexture, ViewRect);
 					FScreenPassTextureViewport OutputViewport(OutputTexture);
 					TShaderMapRef<FRGB8toUYVY8ConvertPS> PixelShader(GlobalShaderMap);
-					FRGB8toUYVY8ConvertPS::FParameters* Parameters = PixelShader->AllocateAndSetParameters(Args.GraphBuilder, ConversionPassArgs.SourceRGBTexture, MediaShaders::RgbToYuvRec709Scaled, MediaShaders::YUVOffset8bits, bDoLinearToSRGB, OutputTexture);
+					FRGB8toUYVY8ConvertPS::FParameters* Parameters = PixelShader->AllocateAndSetParameters(Args.GraphBuilder, ConversionPassArgs.SourceRGBTexture, ConversionMatrix, MediaShaders::YUVOffset8bits, bDoLinearToSRGB, OutputTexture);
 					AddDrawScreenPass(Args.GraphBuilder, RDG_EVENT_NAME("RGBToUYVY 8 bit"), FScreenPassViewInfo(), OutputViewport, InputViewport, VertexShader, PixelShader, Parameters);
 				}
 				break;
@@ -484,7 +486,7 @@ namespace UE::MediaCapture
 					FScreenPassTextureViewport InputViewport(ConversionPassArgs.SourceRGBTexture, ViewRect);
 					FScreenPassTextureViewport OutputViewport(OutputTexture);
 					TShaderMapRef<FRGB10toYUVv210ConvertPS> PixelShader(GlobalShaderMap);
-					FRGB10toYUVv210ConvertPS::FParameters* Parameters = PixelShader->AllocateAndSetParameters(Args.GraphBuilder, ConversionPassArgs.SourceRGBTexture, MediaShaders::RgbToYuvRec709Scaled, MediaShaders::YUVOffset10bits, bDoLinearToSRGB, OutputTexture);
+					FRGB10toYUVv210ConvertPS::FParameters* Parameters = PixelShader->AllocateAndSetParameters(Args.GraphBuilder, ConversionPassArgs.SourceRGBTexture, ConversionMatrix, MediaShaders::YUVOffset10bits, bDoLinearToSRGB, OutputTexture);
 					AddDrawScreenPass(Args.GraphBuilder, RDG_EVENT_NAME("RGBToYUVv210"), FScreenPassViewInfo(), OutputViewport, InputViewport, VertexShader, PixelShader, Parameters);
 				}
 				break;
diff --git a/Engine/Plugins/Media/MediaIOFramework/Source/MediaIOCore/Public/MediaCapture.h b/Engine/Plugins/Media/MediaIOFramework/Source/MediaIOCore/Public/MediaCapture.h
--- a/Engine/Plugins/Media/MediaIOFramework/Source/MediaIOCore/Public/MediaCapture.h
+++ b/Engine/Plugins/Media/MediaIOFramework/Source/MediaIOCore/Public/MediaCapture.h
@@ -403,6 +403,12 @@ public:
 	virtual bool UpdateRenderTargetImpl(UTextureRenderTarget2D* InRenderTarget) { return true; }
 	virtual bool UpdateAudioDeviceImpl(const FAudioDeviceHandle& InAudioDeviceHandle) { return true; }
 
+	/**
+	 * Get the RGB to YUV Conversion Matrix, should be overriden for media outputs that support different color gamuts (ie. BT 2020).
+	 * Will default to returning the RGBToYUV Rec709 matrix.
+	 */
+	virtual const FMatrix& GetRGBToYUVConversionMatrix() const;
+
 	static const TSet<EPixelFormat>& GetSupportedRgbaSwizzleFormats();
 
 	/** Get the name of the media output that created this capture. */
diff --git a/Engine/Source/Runtime/RenderCore/Private/MediaShaders.cpp b/Engine/Source/Runtime/RenderCore/Private/MediaShaders.cpp
--- a/Engine/Source/Runtime/RenderCore/Private/MediaShaders.cpp
+++ b/Engine/Source/Runtime/RenderCore/Private/MediaShaders.cpp
@@ -97,6 +97,14 @@ namespace MediaShaders
 		FPlane(0.000000f, 0.000000f, 0.000000f, 0.000000f)
 	);
 
+	// Inverse of YuvToRgbRec2020Scaled
+	const FMatrix RgbToYuvRec2020Scaled = FMatrix(
+		FPlane( 0.225613123257483f,  0.582280696718123f,  0.0509297094389877f, 0.000000f),
+		FPlane(-0.122655750438889f, -0.316559935835522f,  0.439215686274411f, 0.000000f),
+		FPlane( 0.439215686274411f, -0.403889154894806f, -0.0353265313796045f, 0.000000f),
+		FPlane(0.000000f, 0.000000f, 0.000000f, 0.000000f)
+	);
+
 	const FVector YUVOffset8bits = FVector(0.06274509803921568627f, 0.5019607843137254902f, 0.5019607843137254902f);
 	const FVector YUVOffsetNoScale8bits = FVector(0.0f, 0.5019607843137254902f, 0.5019607843137254902f);
 
diff --git a/Engine/Source/Runtime/RenderCore/Public/MediaShaders.h b/Engine/Source/Runtime/RenderCore/Public/MediaShaders.h
--- a/Engine/Source/Runtime/RenderCore/Public/MediaShaders.h
+++ b/Engine/Source/Runtime/RenderCore/Public/MediaShaders.h
@@ -68,6 +68,9 @@ namespace MediaShaders
 	/** Color transform from RGB to YUV (in Rec. 709 color space, including inversion of range scaling) */
 	RENDERCORE_API extern const FMatrix RgbToYuvRec709Scaled;
 
+	/** Color transform from RGB to YUV (in Rec. 2020 color space, including inversion of range scaling) */
+	RENDERCORE_API extern const FMatrix RgbToYuvRec2020Scaled;
+
 	/** YUV Offset for 8 bit conversion (Computed as 16/255, 128/255, 128/255) */
 	RENDERCORE_API extern const FVector YUVOffset8bits;
 
