[MediaIOCore] Fix texture becoming invalid during garbage collection.

Using IsValid on a weak ptr from the render thread needs the ThreadSafeTest = true, otherwise it will return false during the mark phase of the GC due to the Unreachable flag being set.

diff --git a/Engine/Plugins/Media/MediaIOFramework/Source/MediaIOCore/Private/MediaCaptureSources.h b/Engine/Plugins/Media/MediaIOFramework/Source/MediaIOCore/Private/MediaCaptureSources.h
--- a/Engine/Plugins/Media/MediaIOFramework/Source/MediaIOCore/Private/MediaCaptureSources.h
+++ b/Engine/Plugins/Media/MediaIOFramework/Source/MediaIOCore/Private/MediaCaptureSources.h
@@ -404,9 +404,11 @@ namespace UE::MediaCapture::Private
 
 		virtual FTexture2DRHIRef GetSourceTextureForInput_RenderThread(FRHICommandListImmediate& RHICmdList)
 		{
-			if (RenderTarget.IsValid() && RenderTarget->GetRenderTargetResource())
+			constexpr bool bEvenIfPendingKill = false;
+			constexpr bool bThreadSafeTest = true;
+			if (RenderTarget.IsValid(bEvenIfPendingKill, bThreadSafeTest) && RenderTarget.GetEvenIfUnreachable()->GetRenderTargetResource())
 			{
-				return RenderTarget->GetRenderTargetResource()->GetTextureRenderTarget2DResource()->GetTextureRHI();
+				return RenderTarget.GetEvenIfUnreachable()->GetRenderTargetResource()->GetTextureRenderTarget2DResource()->GetTextureRHI();
 			}
 			return nullptr;
 		}
