Add Media IO input HDR support

diff --git a/Engine/Plugins/Media/AjaMedia/Source/Aja/Private/InputChannel.cpp b/Engine/Plugins/Media/AjaMedia/Source/Aja/Private/InputChannel.cpp
--- a/Engine/Plugins/Media/AjaMedia/Source/Aja/Private/InputChannel.cpp
+++ b/Engine/Plugins/Media/AjaMedia/Source/Aja/Private/InputChannel.cpp
@@ -188,7 +188,7 @@ namespace AJA
 			NTV2DeviceID DeviceId = GetDevice().GetDeviceID();
 			if (!::NTV2DeviceCanDoCapture(DeviceId))
 			{
-				UE_LOG(LogTemp, Error, TEXT("AutoCirculate: The device '%S' couldn't not capture."), GetDevice().GetDisplayName().c_str());
+				UE_LOG(LogAjaCore, Error, TEXT("AutoCirculate: The device '%S' couldn't not capture."), GetDevice().GetDisplayName().c_str());
 				return false;
 			}
 
@@ -218,7 +218,7 @@ namespace AJA
 				bRunning = GetDevice().AutoCirculateGetStatus(Channel, ChannelStatus);
 				if (!bRunning)
 				{
-					UE_LOG(LogTemp, Error, TEXT("AutoCirculate: Can't get the status for channel %d on device %S.\n"), uint32_t(Channel) + 1, GetDevice().GetDisplayName().c_str());
+					UE_LOG(LogAjaCore, Error, TEXT("AutoCirculate: Can't get the status for channel %d on device %S.\n"), uint32_t(Channel) + 1, GetDevice().GetDisplayName().c_str());
 					bRunning = false;
 				}
 
@@ -231,7 +231,7 @@ namespace AJA
 				const DWORD TimeoutMilli = 2000;
 				if (Counter * SleepMilli > TimeoutMilli)
 				{
-					UE_LOG(LogTemp, Error, TEXT("AutoCirculate: Can't get the Channel running for channel %d on device %S.\n"), uint32_t(Channel) + 1, GetDevice().GetDisplayName().c_str());
+					UE_LOG(LogAjaCore, Error, TEXT("AutoCirculate: Can't get the Channel running for channel %d on device %S.\n"), uint32_t(Channel) + 1, GetDevice().GetDisplayName().c_str());
 					bRunning = false;
 				}
 
@@ -262,7 +262,7 @@ namespace AJA
 				bRunning = GetDevice().AutoCirculateGetStatus(Channel, ChannelStatus);
 				if (!bRunning)
 				{
-					UE_LOG(LogTemp, Error, TEXT("AutoCirculate: Can't get the status for channel %d on device %S.\n"), uint32_t(Channel) + 1, GetDevice().GetDisplayName().c_str());
+					UE_LOG(LogAjaCore, Error, TEXT("AutoCirculate: Can't get the status for channel %d on device %S.\n"), uint32_t(Channel) + 1, GetDevice().GetDisplayName().c_str());
 					break;
 				}
 
@@ -325,6 +325,8 @@ namespace AJA
 
 
 					AJAVideoFrameData VideoData;
+					VideoData.HDROptions = GetOptions().HDROptions;
+					
 					if (VideoBuffer)
 					{
 #if AJA_TEST_MEMORY_BUFFER
@@ -343,7 +345,7 @@ namespace AJA
 					bRunning = GetDevice().AutoCirculateTransfer(Channel, Transfer);
 					if (!bRunning)
 					{
-						UE_LOG(LogTemp, Error, TEXT("AutoCirculate: Can't transfer the buffer for channel %d on device %S.\n"), uint32_t(Channel) + 1, GetDevice().GetDisplayName().c_str());
+						UE_LOG(LogAjaCore, Error, TEXT("AutoCirculate: Can't transfer the buffer for channel %d on device %S.\n"), uint32_t(Channel) + 1, GetDevice().GetDisplayName().c_str());
 						break;
 					}
 
@@ -452,7 +454,7 @@ namespace AJA
 						}
 						else
 						{
-							UE_LOG(LogTemp, Warning,  TEXT("AutoCirculate: The device '%S' missed an interrupt signal."), GetDevice().GetDisplayName().c_str());
+							UE_LOG(LogAjaCore, Warning,  TEXT("AutoCirculate: The device '%S' missed an interrupt signal."), GetDevice().GetDisplayName().c_str());
 						}
 					}
 				}
@@ -495,7 +497,7 @@ namespace AJA
 			NTV2DeviceID DeviceId = GetDevice().GetDeviceID();
 			if (!::NTV2DeviceCanDoCapture(DeviceId))
 			{
-				UE_LOG(LogTemp, Error, TEXT("PingPong: The device '%S' couldn't not capture."), GetDevice().GetDisplayName().c_str());
+				UE_LOG(LogAjaCore, Error, TEXT("PingPong: The device '%S' couldn't not capture."), GetDevice().GetDisplayName().c_str());
 				return false;
 			}
 
@@ -558,7 +560,7 @@ namespace AJA
 					}
 					else
 					{
-						UE_LOG(LogTemp, Warning,  TEXT("PingPong: The device '%S' missed an interrupt signal."), GetDevice().GetDisplayName().c_str());
+						UE_LOG(LogAjaCore, Warning,  TEXT("PingPong: The device '%S' missed an interrupt signal."), GetDevice().GetDisplayName().c_str());
 						::Sleep(0);
 						continue;
 					}
@@ -566,7 +568,7 @@ namespace AJA
 
 				if (!bRunning)
 				{
-					UE_LOG(LogTemp, Error, TEXT("PingPong: Can't wait for the input field for channel %d on device %S.\n"), uint32_t(Channel) + 1, GetDevice().GetDisplayName().c_str());
+					UE_LOG(LogAjaCore, Error, TEXT("PingPong: Can't wait for the input field for channel %d on device %S.\n"), uint32_t(Channel) + 1, GetDevice().GetDisplayName().c_str());
 					break;
 				}
 
@@ -688,6 +690,8 @@ namespace AJA
 
 
 				AJAVideoFrameData VideoData;
+				VideoData.HDROptions = GetOptions().HDROptions;
+				
 				if (VideoBuffer)
 				{
 #if AJA_TEST_MEMORY_BUFFER
@@ -706,7 +710,7 @@ namespace AJA
 
 				if (!bRunning)
 				{
-					UE_LOG(LogTemp, Error, TEXT("PingPong: Can't do the DMA frame transfer for channel %d on device %S.\n"), uint32_t(Channel) + 1, GetDevice().GetDisplayName().c_str());
+					UE_LOG(LogAjaCore, Error, TEXT("PingPong: Can't do the DMA frame transfer for channel %d on device %S.\n"), uint32_t(Channel) + 1, GetDevice().GetDisplayName().c_str());
 					break;
 				}
 
@@ -851,7 +855,7 @@ namespace AJA
 					FoundFormat = Helpers::GetInputVideoFormat(GetDevicePtr(), GetOptions().TransportType, Channel, InputSource, VideoFormat, false, FailureReason);
 					if (Counter * SleepMilli > TimeoutMilli)
 					{
-						UE_LOG(LogTemp, Error, TEXT("%S: Can't get the Channel running for channel %d on device %S.\n"), ModeName->c_str(), uint32_t(Channel) + 1, GetDevice().GetDisplayName().c_str());
+						UE_LOG(LogAjaCore, Error, TEXT("%S: Can't get the Channel running for channel %d on device %S.\n"), ModeName->c_str(), uint32_t(Channel) + 1, GetDevice().GetDisplayName().c_str());
 						break;
 					}
 
@@ -877,7 +881,7 @@ namespace AJA
 					// Todo: Warn once, display signal lost symbol 
 					if (!bAutoDetectFormat)
 					{
-						UE_LOG(LogTemp, Error, TEXT("%S: Could not detect format for channel %d on device %S. %S\n")
+						UE_LOG(LogAjaCore, Error, TEXT("%S: Could not detect format for channel %d on device %S. %S\n")
 							, ModeName->c_str()
 							, uint32_t(Channel) + 1
 							, GetDevice().GetDisplayName().c_str()
@@ -928,7 +932,7 @@ namespace AJA
 				}
 				else
 				{
-					UE_LOG(LogTemp, Error, TEXT("%S: The VideoFormat changed for channel %d on device %S. %S\n")
+					UE_LOG(LogAjaCore, Error, TEXT("%S: The VideoFormat changed for channel %d on device %S. %S\n")
 						, ModeName->c_str()
 						, uint32_t(Channel) + 1
 						, GetDevice().GetDisplayName().c_str()
diff --git a/Engine/Plugins/Media/AjaMedia/Source/Aja/Public/AJALib.h b/Engine/Plugins/Media/AjaMedia/Source/Aja/Public/AJALib.h
--- a/Engine/Plugins/Media/AjaMedia/Source/Aja/Public/AJALib.h
+++ b/Engine/Plugins/Media/AjaMedia/Source/Aja/Public/AJALib.h
@@ -293,6 +293,52 @@ namespace AJA
 	private:
 		Private::TimecodeChannel* Channel;
 	};
+	
+	
+	/**
+	 * HDR Transfer function.
+	 */
+	enum class EAjaHDRMetadataEOTF : uint8
+	{
+		SDR,
+		HLG,
+		PQ,
+		Unspecified
+	};
+	
+	/**
+	 * HDR Color Gamut.
+	 */
+	enum class EAjaHDRMetadataGamut : uint8
+	{
+		Rec709,
+		Rec2020,
+		Invalid
+	};
+
+	/**
+	 * HDR Luminance.
+	 */
+	enum class EAjaHDRMetadataLuminance : uint8
+	{
+		YCbCr,
+		ICtCp
+	};
+
+	/**
+	 * Set of metadata describing a HDR video signal.
+	 */
+	struct AJACORE_API FAjaHDROptions
+	{
+		/** Transfer function to use for converting the video signal to an optical signal. */
+		EAjaHDRMetadataEOTF EOTF = EAjaHDRMetadataEOTF::SDR;
+
+		/** The color gamut of the video signal. */
+		EAjaHDRMetadataGamut Gamut = EAjaHDRMetadataGamut::Rec709;
+
+		/** Color representation format of the video signal. */
+		EAjaHDRMetadataLuminance Luminance = EAjaHDRMetadataLuminance::YCbCr;
+	};
 
 	/* AJAInputFrameData definition
 	*****************************************************************************/
@@ -343,6 +389,7 @@ namespace AJA
 		uint32_t Height;
 		EPixelFormat PixelFormat;
 		bool bIsProgressivePicture;
+		FAjaHDROptions HDROptions;
 	};
 
 	struct AJACORE_API AJARequestInputBufferData
@@ -364,52 +411,6 @@ namespace AJA
 		uint8_t* VideoBuffer;
 	};
 
-	
-	/**
-	 * HDR Transfer function.
-	 */
-	enum class EAjaHDRMetadataEOTF : uint8
-	{
-		SDR,
-		HLG,
-		PQ,
-		Unspecified
-	};
-	
-	/**
-	 * HDR Color Gamut.
-	 */
-	enum class EAjaHDRMetadataGamut : uint8
-	{
-		Rec709,
-		Rec2020,
-		Invalid
-	};
-
-	/**
-	 * HDR Luminance.
-	 */
-	enum class EAjaHDRMetadataLuminance : uint8
-	{
-		YCbCr,
-		ICtCp
-	};
-
-	/**
-	 * Set of metadata describing a HDR video signal.
-	 */
-	struct AJACORE_API FAjaHDROptions
-	{
-		/** Transfer function to use for converting the video signal to an optical signal. */
-		EAjaHDRMetadataEOTF EOTF = EAjaHDRMetadataEOTF::SDR;
-
-		/** The color gamut of the video signal. */
-		EAjaHDRMetadataGamut Gamut = EAjaHDRMetadataGamut::Rec709;
-
-		/** Color representation format of the video signal. */
-		EAjaHDRMetadataLuminance Luminance = EAjaHDRMetadataLuminance::YCbCr;
-	};
-
 	struct AJACORE_API IAJAInputOutputChannelCallbackInterface : IAJASyncChannelCallbackInterface
 	{
 		IAJAInputOutputChannelCallbackInterface();
diff --git a/Engine/Plugins/Media/AjaMedia/Source/AjaMedia/Private/AjaMediaDefinitions.cpp b/Engine/Plugins/Media/AjaMedia/Source/AjaMedia/Private/AjaMediaDefinitions.cpp
new file mode 100644
--- /dev/null
+++ b/Engine/Plugins/Media/AjaMedia/Source/AjaMedia/Private/AjaMediaDefinitions.cpp
@@ -0,0 +1,28 @@
+// Copyright Epic Games, Inc. All Rights Reserved.
+
+#include "AjaMediaDefinitions.h"
+
+#include "AjaLib.h"
+
+namespace UE::AjaMedia
+{
+	FAjaMediaHDROptions MakeAjaMediaHDROptions(const AJA::FAjaHDROptions& HDROptions)
+	{
+		FAjaMediaHDROptions HDRMetadata;
+
+		HDRMetadata.Gamut = (EAjaHDRMetadataGamut)HDROptions.Gamut;
+		HDRMetadata.EOTF = (EAjaHDRMetadataEOTF)HDROptions.EOTF;
+
+		return HDRMetadata;
+	}
+
+	AJA::FAjaHDROptions MakeAjaHDROptions(const FAjaMediaHDROptions& HDROptions)
+	{
+		AJA::FAjaHDROptions HDRMetadata;
+
+		HDRMetadata.Gamut = (AJA::EAjaHDRMetadataGamut)HDROptions.Gamut;
+		HDRMetadata.EOTF = (AJA::EAjaHDRMetadataEOTF)HDROptions.EOTF;
+
+		return HDRMetadata;
+	}
+}
diff --git a/Engine/Plugins/Media/AjaMedia/Source/AjaMedia/Private/Assets/AjaMediaSource.cpp b/Engine/Plugins/Media/AjaMedia/Source/AjaMedia/Private/Assets/AjaMediaSource.cpp
--- a/Engine/Plugins/Media/AjaMedia/Source/AjaMedia/Private/Assets/AjaMediaSource.cpp
+++ b/Engine/Plugins/Media/AjaMedia/Source/AjaMedia/Private/Assets/AjaMediaSource.cpp
@@ -16,7 +16,6 @@ UAjaMediaSource::UAjaMediaSource()
 	, MaxNumAudioFrameBuffer(8)
 	, bCaptureVideo(true)
 	, ColorFormat(EAjaMediaSourceColorFormat::YUV2_8bit)
-	, bIsSRGBInput(false)
 	, MaxNumVideoFrameBuffer(8)
 	, bLogDropFrame(true)
 	, bEncodeTimecodeInTexel(false)
@@ -60,11 +59,6 @@ bool UAjaMediaSource::GetMediaOption(const FName& Key, bool DefaultValue) const
 	{
 		return bEncodeTimecodeInTexel;
 	}
-	if (Key == AjaMediaOption::SRGBInput)
-	{
-		return bIsSRGBInput;
-	}
-
 
 	return Super::GetMediaOption(Key, DefaultValue);
 }
@@ -165,7 +159,6 @@ bool UAjaMediaSource::HasMediaOption(const FName& Key) const
 		(Key == AjaMediaOption::MaxAudioFrameBuffer) ||
 		(Key == AjaMediaOption::AjaVideoFormat) ||
 		(Key == AjaMediaOption::ColorFormat) ||
-		(Key == AjaMediaOption::SRGBInput) ||
 		(Key == AjaMediaOption::MaxVideoFrameBuffer) ||
 		(Key == AjaMediaOption::LogDropFrame) ||
 		(Key == AjaMediaOption::EncodeTimecodeInTexel)
diff --git a/Engine/Plugins/Media/AjaMedia/Source/AjaMedia/Private/Player/AjaMediaPlayer.cpp b/Engine/Plugins/Media/AjaMedia/Source/AjaMedia/Private/Player/AjaMediaPlayer.cpp
--- a/Engine/Plugins/Media/AjaMedia/Source/AjaMedia/Private/Player/AjaMediaPlayer.cpp
+++ b/Engine/Plugins/Media/AjaMedia/Source/AjaMedia/Private/Player/AjaMediaPlayer.cpp
@@ -6,6 +6,7 @@
 #include "Aja/Aja.h"
 #include "AjaMediaAudioSample.h"
 #include "AjaMediaBinarySample.h"
+#include "AjaMediaDefinitions.h"
 #include "AjaMediaPrivate.h"
 #include "AjaMediaSettings.h"
 #include "AjaMediaTextureSample.h"
@@ -17,6 +18,7 @@
 #include "MediaIOCoreEncodeTime.h"
 #include "MediaIOCoreFileWriter.h"
 #include "MediaIOCoreSamples.h"
+#include "MediaIOCoreUtilities.h"
 #include "HAL/PlatformAtomics.h"
 #include "HAL/PlatformProcess.h"
 #include "IAjaMediaModule.h"
@@ -119,6 +121,11 @@ bool FAjaMediaPlayer::Open(const FString& Url, const IMediaOptions* Options)
 	const bool bAutoDetectTimecode = Timecode == EMediaIOAutoDetectableTimecodeFormat::Auto;
 	const bool bAutoDetectVideoFormat = bAutoDetect;
 
+	bOverrideSourceEncoding = Options->GetMediaOption(UE::CaptureCardMediaSource::OverrideSourceEncoding, true);
+	OverrideSourceEncoding = (ETextureSourceEncoding) Options->GetMediaOption(UE::CaptureCardMediaSource::SourceEncoding, (int64) ETextureSourceEncoding::TSE_Linear);
+	bOverrideSourceColorSpace = Options->GetMediaOption(UE::CaptureCardMediaSource::OverrideSourceColorSpace, true);
+	OverrideSourceColorSpace =  (ETextureColorSpace) Options->GetMediaOption(UE::CaptureCardMediaSource::SourceColorSpace, (int64) ETextureColorSpace::TCS_None);
+
 	if (!bAutoDetectTimecode)
 	{
 		TimecodeFormat = UE::MediaIO::FromAutoDetectableTimecodeFormat(Timecode);
@@ -463,7 +470,6 @@ bool FAjaMediaPlayer::OnInputFrameReceived(const AJA::AJAInputFrameData& InInput
 {
 	SCOPE_CYCLE_COUNTER(STAT_AJA_MediaPlayer_ProcessFrame);
 
-
 	if ((AjaThreadNewState != EMediaState::Playing) && (AjaThreadNewState != EMediaState::Paused))
 	{
 		return false;
@@ -625,9 +631,59 @@ bool FAjaMediaPlayer::OnInputFrameReceived(const AJA::AJAInputFrameData& InInput
 			bAjaWriteOutputRawDataCmdEnable = false;
 		}
 
+		UE::MediaIOCore::FColorFormatArgs ColorFormat;
+		FAjaMediaHDROptions HDROptions = UE::AjaMedia::MakeAjaMediaHDROptions(InVideoFrame.HDROptions);
+		
+		if (bOverrideSourceEncoding)
+		{
+			ColorFormat.Encoding = (UE::Color::EEncoding) OverrideSourceEncoding;
+		}
+		else
+		{
+			const FTimespan TimeBetweenLogs = FTimespan::FromSeconds(5);
+
+			switch (HDROptions.EOTF)
+			{
+			case EAjaHDRMetadataEOTF::SDR:
+				ColorFormat.Encoding = UE::Color::EEncoding::sRGB; // Missing support for BT.1886 so we default to sRGB instead.
+				break;
+			case EAjaHDRMetadataEOTF::PQ:
+				ColorFormat.Encoding = UE::Color::EEncoding::ST2084;
+				break;
+			case EAjaHDRMetadataEOTF::HLG:
+				UE_MEDIA_IO_LOG_THROTTLE(LogAjaMedia, Warning, TimeBetweenLogs, TEXT("HLG EOTF is not supported at the moment. Defaulting to sRGB."));
+				ColorFormat.Encoding = UE::Color::EEncoding::sRGB;
+				break;
+			default:
+				checkNoEntry();
+				ColorFormat.Encoding = UE::Color::EEncoding::None;			
+			}
+		}
+		
+		if (bOverrideSourceColorSpace)
+		{
+			ColorFormat.ColorSpace = (UE::Color::EColorSpace) OverrideSourceColorSpace;
+		}
+		else
+		{
+			switch (HDROptions.Gamut)
+			{
+			case EAjaHDRMetadataGamut::Rec709:
+				ColorFormat.ColorSpace = UE::Color::EColorSpace::sRGB;
+				break;
+			case EAjaHDRMetadataGamut::Rec2020:
+				ColorFormat.ColorSpace = UE::Color::EColorSpace::Rec2020;
+				break;
+			default:
+				checkNoEntry();
+				ColorFormat.ColorSpace = UE::Color::EColorSpace::sRGB;
+				break;
+			}
+		}
+
 		if (AjaThreadCurrentTextureSample.IsValid())
 		{
-			if (AjaThreadCurrentTextureSample->SetProperties(InVideoFrame.Stride, InVideoFrame.Width, InVideoFrame.Height, VideoSampleFormat, DecodedTime, VideoFrameRate, DecodedTimecode, bIsSRGBInput))
+			if (AjaThreadCurrentTextureSample->SetProperties(InVideoFrame.Stride, InVideoFrame.Width, InVideoFrame.Height, VideoSampleFormat, DecodedTime, VideoFrameRate, DecodedTimecode, ColorFormat))
 			{
 				if (CanUseGPUTextureTransfer())
 				{
@@ -662,7 +718,7 @@ bool FAjaMediaPlayer::OnInputFrameReceived(const AJA::AJAInputFrameData& InInput
 				if (InVideoFrame.bIsProgressivePicture)
 				{
 					auto TextureSample = TextureSamplePool->AcquireShared();
-					if (TextureSample->InitializeProgressive(InVideoFrame, VideoSampleFormat, DecodedTime, VideoFrameRate, DecodedTimecode, bIsSRGBInput))
+					if (TextureSample->InitializeProgressive(InVideoFrame, VideoSampleFormat, DecodedTime, VideoFrameRate, DecodedTimecode, ColorFormat))
 					{
 						Samples->AddVideo(TextureSample);
 					}
@@ -685,7 +741,7 @@ bool FAjaMediaPlayer::OnInputFrameReceived(const AJA::AJAInputFrameData& InInput
 						DecodedTime,
 						VideoFrameRate,
 						DecodedTimecode,
-						bIsSRGBInput ? UE::Color::EColorSpace::sRGB : UE::Color::EColorSpace::None
+						ColorFormat
 					};
 
 					const TArray<TSharedRef<FMediaIOCoreTextureSampleBase>> DeinterlacedSamples = Deinterlace(FrameInfo);
diff --git a/Engine/Plugins/Media/AjaMedia/Source/AjaMedia/Private/Player/AjaMediaPlayer.h b/Engine/Plugins/Media/AjaMedia/Source/AjaMedia/Private/Player/AjaMediaPlayer.h
--- a/Engine/Plugins/Media/AjaMedia/Source/AjaMedia/Private/Player/AjaMediaPlayer.h
+++ b/Engine/Plugins/Media/AjaMedia/Source/AjaMedia/Private/Player/AjaMediaPlayer.h
@@ -192,6 +192,18 @@ private:
 
 	EAjaMediaSourceColorFormat AjaColorFormat = EAjaMediaSourceColorFormat::YUV2_8bit;
 
+	/** Whether to override the source encoding or to use the metadata embedded in the ancillary data of the signal. */
+	bool bOverrideSourceEncoding = true;
+
+	/** Encoding of the source texture. */
+	ETextureSourceEncoding OverrideSourceEncoding = ETextureSourceEncoding::TSE_Linear;
+
+	/** Whether to override the source color space or to use the metadata embedded in the ancillary data of the signal. */
+	bool bOverrideSourceColorSpace = true;
+
+	/** Color space of the source texture. */
+	ETextureColorSpace OverrideSourceColorSpace = ETextureColorSpace::TCS_None;
+
 	/** Device provider used to autodetect input format. */
 	TPimplPtr<class FAjaDeviceProvider> DeviceProvider;
 
diff --git a/Engine/Plugins/Media/AjaMedia/Source/AjaMedia/Private/Player/AjaMediaTextureSample.h b/Engine/Plugins/Media/AjaMedia/Source/AjaMedia/Private/Player/AjaMediaTextureSample.h
--- a/Engine/Plugins/Media/AjaMedia/Source/AjaMedia/Private/Player/AjaMediaTextureSample.h
+++ b/Engine/Plugins/Media/AjaMedia/Source/AjaMedia/Private/Player/AjaMediaTextureSample.h
@@ -24,7 +24,7 @@ public:
 	 * @param InTimecode The sample timecode if available.
 	 * @param bInIsSRGB Whether the sample is in sRGB space.
 	 */
-	bool InitializeProgressive(const AJA::AJAVideoFrameData& InVideoData, EMediaTextureSampleFormat InSampleFormat, FTimespan InTime, const FFrameRate& InFrameRate, const TOptional<FTimecode>& InTimecode, bool bInIsSRGB)
+	bool InitializeProgressive(const AJA::AJAVideoFrameData& InVideoData, EMediaTextureSampleFormat InSampleFormat, FTimespan InTime, const FFrameRate& InFrameRate, const TOptional<FTimecode>& InTimecode, const UE::MediaIOCore::FColorFormatArgs& InColorFormatArgs)
 	{
 		return Super::Initialize(InVideoData.VideoBuffer
 			, InVideoData.VideoBufferSize
@@ -35,7 +35,7 @@ public:
 			, InTime
 			, InFrameRate
 			, InTimecode
-			, bInIsSRGB);
+			, InColorFormatArgs);
 	}
 
 	/**
@@ -49,7 +49,7 @@ public:
 	 * @param bEven Only take the even frame from the image.
 	 * @param bInIsSRGB Whether the sample is in sRGB space.
 	 */
-	bool InitializeInterlaced_Halfed(const AJA::AJAVideoFrameData& InVideoData, EMediaTextureSampleFormat InSampleFormat, FTimespan InTime, const FFrameRate& InFrameRate, const TOptional<FTimecode>& InTimecode, bool bInEven, bool bInIsSRGB)
+	bool InitializeInterlaced_Halfed(const AJA::AJAVideoFrameData& InVideoData, EMediaTextureSampleFormat InSampleFormat, FTimespan InTime, const FFrameRate& InFrameRate, const TOptional<FTimecode>& InTimecode, bool bInEven, const UE::MediaIOCore::FColorFormatArgs& InColorFormatArgs)
 	{
 		return Super::InitializeWithEvenOddLine(bInEven
 			, InVideoData.VideoBuffer
@@ -61,17 +61,7 @@ public:
 			, InTime
 			, InFrameRate
 			, InTimecode
-			, bInIsSRGB);
-	}
-
-	/**
-	 * Get YUV to RGB conversion matrix
-	 *
-	 * @return MediaIOCore Yuv To Rgb matrix
-	 */
-	virtual const FMatrix& GetYUVToRGBMatrix() const override
-	{
-		return MediaShaders::YuvToRgbRec709Scaled;
+			, InColorFormatArgs);
 	}
 };
 
diff --git a/Engine/Plugins/Media/AjaMedia/Source/AjaMedia/Public/AjaMediaDefinitions.h b/Engine/Plugins/Media/AjaMedia/Source/AjaMedia/Public/AjaMediaDefinitions.h
new file mode 100644
--- /dev/null
+++ b/Engine/Plugins/Media/AjaMedia/Source/AjaMedia/Public/AjaMediaDefinitions.h
@@ -0,0 +1,59 @@
+// Copyright Epic Games, Inc. All Rights Reserved.
+
+#pragma once
+
+#include "UObject/Object.h"
+#include "AjaMediaDefinitions.generated.h"
+
+namespace AJA
+{
+	struct FAjaHDROptions;
+}
+
+/**
+ * HDR Transfer function.
+ * Must match AJA::EAjaHDRMetadataEOTF.
+ */
+UENUM()
+enum class EAjaHDRMetadataEOTF : uint8
+{
+	SDR,
+	HLG,
+	PQ
+};
+	
+/**
+ * HDR Color Gamut.
+ * Must match AJA::EAjaHDRMetadataGamut.
+ */
+UENUM()
+enum class EAjaHDRMetadataGamut : uint8
+{
+	Rec709,
+	Rec2020
+};
+
+/**
+ * Set of metadata describing a HDR video signal.
+ */
+USTRUCT()
+struct AJAMEDIA_API FAjaMediaHDROptions
+{
+	GENERATED_BODY()
+
+	/** Transfer function to use for converting the video signal to an optical signal. */
+	UPROPERTY(EditAnywhere, Category = "HDR")
+	EAjaHDRMetadataEOTF EOTF = EAjaHDRMetadataEOTF::SDR;
+
+	/** The color gamut of the video signal. */
+	UPROPERTY(EditAnywhere, Category = "HDR")
+	EAjaHDRMetadataGamut Gamut = EAjaHDRMetadataGamut::Rec709;
+};
+
+namespace UE::AjaMedia
+{
+	/** Create Aja HDR UStruct from the AJALib's hdr struct. */
+	FAjaMediaHDROptions MakeAjaMediaHDROptions(const AJA::FAjaHDROptions& HDROptions);
+	/** Create AJALib's hdr struct from the Aja HDR UStruct. */
+	AJA::FAjaHDROptions MakeAjaHDROptions(const FAjaMediaHDROptions& HDROptions);
+}
diff --git a/Engine/Plugins/Media/AjaMedia/Source/AjaMedia/Public/AjaMediaSource.h b/Engine/Plugins/Media/AjaMedia/Source/AjaMedia/Public/AjaMediaSource.h
--- a/Engine/Plugins/Media/AjaMedia/Source/AjaMedia/Public/AjaMediaSource.h
+++ b/Engine/Plugins/Media/AjaMedia/Source/AjaMedia/Public/AjaMediaSource.h
@@ -3,8 +3,10 @@
 #pragma once
 
 #include "CaptureCardMediaSource.h"
+#include "Engine/TextureDefines.h"
 #include "MediaIOCoreDefinitions.h"
 
+
 #include "AjaMediaSource.generated.h"
 
 /**
@@ -95,13 +97,15 @@ public:
 	UPROPERTY(BlueprintReadOnly, EditAnywhere, Category="Video", meta=(EditCondition="bCaptureVideo"))
 	EAjaMediaSourceColorFormat ColorFormat;
 
+#if WITH_EDITORONLY_DATA
 	/** 
 	 * Whether the video input is in sRGB color space.
 	 * A sRGB to Linear conversion will be applied resulting in a texture in linear space.
 	 * @Note If the texture is not in linear space, it won't look correct in the editor. Another pass will be required either through Composure or other means.
 	 */
-	UPROPERTY(BlueprintReadOnly, EditAnywhere, Category = "Video")
-	bool bIsSRGBInput;
+	UPROPERTY(meta=(DeprecatedProperty, DeprecationMessage = "Use ColorEncoding instead."))
+	bool bIsSRGBInput_DEPRECATED = false;
+#endif
 
 	/** Maximum number of video frames to buffer. */
 	UPROPERTY(BlueprintReadOnly, EditAnywhere, AdvancedDisplay, Category="Video", meta=(EditCondition="bCaptureVideo", ClampMin="1", ClampMax="32"))
diff --git a/Engine/Plugins/Media/AjaMedia/Source/AjaMediaOutput/AjaMediaOutput.Build.cs b/Engine/Plugins/Media/AjaMedia/Source/AjaMediaOutput/AjaMediaOutput.Build.cs
--- a/Engine/Plugins/Media/AjaMedia/Source/AjaMediaOutput/AjaMediaOutput.Build.cs
+++ b/Engine/Plugins/Media/AjaMedia/Source/AjaMediaOutput/AjaMediaOutput.Build.cs
@@ -18,6 +18,7 @@ namespace UnrealBuildTool.Rules
 				new string[]
 				{
 					"AjaCore",
+					"ColorManagement",
 					"Core",
 					"CoreUObject",
 					"Engine",
diff --git a/Engine/Plugins/Media/AjaMedia/Source/AjaMediaOutput/Public/AjaMediaOutput.h b/Engine/Plugins/Media/AjaMedia/Source/AjaMediaOutput/Public/AjaMediaOutput.h
--- a/Engine/Plugins/Media/AjaMedia/Source/AjaMediaOutput/Public/AjaMediaOutput.h
+++ b/Engine/Plugins/Media/AjaMedia/Source/AjaMediaOutput/Public/AjaMediaOutput.h
@@ -3,6 +3,8 @@
 #pragma once
 
 #include "MediaOutput.h"
+
+#include "AjaMediaDefinitions.h"
 #include "MediaIOCoreDefinitions.h"
 
 #include "AjaMediaOutput.generated.h"
@@ -34,46 +36,6 @@ enum class EAjaMediaOutputChannelConfiguration : uint8
 	CH_16 = 16 UMETA(DisplayName = "16 Channels")
 };
 
-/**
- * HDR Transfer function.
- * Must match AJA::EAjaHDRMetadataEOTF.
- */
-UENUM()
-enum class EAjaHDRMetadataEOTF : uint8
-{
-	SDR,
-	HLG,
-	PQ
-};
-	
-/**
- * HDR Color Gamut.
- * Must match AJA::EAjaHDRMetadataGamut.
- */
-UENUM()
-enum class EAjaHDRMetadataGamut : uint8
-{
-	Rec709,
-	Rec2020
-};
-
-/**
- * Set of metadata describing a HDR video signal.
- */
-USTRUCT()
-struct AJAMEDIAOUTPUT_API FAjaMediaHDROptions
-{
-	GENERATED_BODY()
-
-	/** Transfer function to use for converting the video signal to an optical signal. */
-	UPROPERTY(EditAnywhere, Category = "HDR")
-	EAjaHDRMetadataEOTF EOTF = EAjaHDRMetadataEOTF::SDR;
-
-	/** The color gamut of the video signal. */
-	UPROPERTY(EditAnywhere, Category = "HDR")
-	EAjaHDRMetadataGamut Gamut = EAjaHDRMetadataGamut::Rec709;
-};
-
 /**
  * Output information for an aja media capture.
  * @note	'Frame Buffer Pixel Format' must be set to at least 8 bits of alpha to enabled the Key.
diff --git a/Engine/Plugins/Media/BlackmagicMedia/Source/BlackmagicMedia/Private/Assets/BlackmagicMediaSource.cpp b/Engine/Plugins/Media/BlackmagicMedia/Source/BlackmagicMedia/Private/Assets/BlackmagicMediaSource.cpp
--- a/Engine/Plugins/Media/BlackmagicMedia/Source/BlackmagicMedia/Private/Assets/BlackmagicMediaSource.cpp
+++ b/Engine/Plugins/Media/BlackmagicMedia/Source/BlackmagicMedia/Private/Assets/BlackmagicMediaSource.cpp
@@ -14,7 +14,6 @@ UBlackmagicMediaSource::UBlackmagicMediaSource()
 	, MaxNumAudioFrameBuffer(8)
 	, bCaptureVideo(true)
 	, ColorFormat(EBlackmagicMediaSourceColorFormat::YUV8)
-	, bIsSRGBInput(false)
 	, MaxNumVideoFrameBuffer(8)
 	, bLogDropFrame(false)
 	, bEncodeTimecodeInTexel(false)
@@ -41,7 +40,6 @@ bool UBlackmagicMediaSource::GetMediaOption(const FName& Key, bool DefaultValue)
 	if (Key == BlackmagicMediaOption::CaptureVideo) { return bCaptureVideo; }
 	if (Key == BlackmagicMediaOption::LogDropFrame) { return bLogDropFrame; }
 	if (Key == BlackmagicMediaOption::EncodeTimecodeInTexel) { return bEncodeTimecodeInTexel; }
-	if (Key == BlackmagicMediaOption::SRGBInput) { return bIsSRGBInput; }
 
 	return Super::GetMediaOption(Key, DefaultValue);
 }
@@ -78,8 +76,7 @@ bool UBlackmagicMediaSource::HasMediaOption(const FName& Key) const
 	if (   Key == BlackmagicMediaOption::CaptureAudio
 		|| Key == BlackmagicMediaOption::CaptureVideo
 		|| Key == BlackmagicMediaOption::LogDropFrame
-		|| Key == BlackmagicMediaOption::EncodeTimecodeInTexel
-		|| Key == BlackmagicMediaOption::SRGBInput)
+		|| Key == BlackmagicMediaOption::EncodeTimecodeInTexel)
 	{
 		return true;
 	}
diff --git a/Engine/Plugins/Media/BlackmagicMedia/Source/BlackmagicMedia/Private/BlackmagicMediaDefinitions.cpp b/Engine/Plugins/Media/BlackmagicMedia/Source/BlackmagicMedia/Private/BlackmagicMediaDefinitions.cpp
new file mode 100644
--- /dev/null
+++ b/Engine/Plugins/Media/BlackmagicMedia/Source/BlackmagicMedia/Private/BlackmagicMediaDefinitions.cpp
@@ -0,0 +1,48 @@
+// Copyright Epic Games, Inc. All Rights Reserved.
+
+#include "BlackmagicMediaDefinitions.h"
+
+#include "BlackmagicLib.h"
+
+namespace UE::BlackmagicMedia
+{
+	FBlackmagicMediaHDROptions MakeBlackmagicMediaHDROptions(const BlackmagicDesign::FHDRMetaData& HDROptions)
+	{
+		FBlackmagicMediaHDROptions HDRMetadata;
+
+		switch (HDROptions.ColorSpace)
+		{
+		case BlackmagicDesign::EHDRMetaDataColorspace::Rec601:
+			// Fallthrough
+		case BlackmagicDesign::EHDRMetaDataColorspace::Rec709:
+			HDRMetadata.Gamut = EBlackmagicHDRMetadataGamut::Rec709;
+			break;
+		case BlackmagicDesign::EHDRMetaDataColorspace::Rec2020:
+			HDRMetadata.Gamut = EBlackmagicHDRMetadataGamut::Rec2020;
+			break;
+		}
+
+		HDRMetadata.EOTF = (EBlackmagicHDRMetadataEOTF)HDROptions.EOTF;
+
+		return HDRMetadata;
+	}
+
+	BlackmagicDesign::FHDRMetaData MakeBlackmagicHDROptions(const FBlackmagicMediaHDROptions& HDROptions)
+	{
+		BlackmagicDesign::FHDRMetaData HDRMetadata;
+
+		switch (HDROptions.Gamut)
+		{
+		case EBlackmagicHDRMetadataGamut::Rec709:
+			HDRMetadata.ColorSpace = BlackmagicDesign::EHDRMetaDataColorspace::Rec709;
+			break;
+		case EBlackmagicHDRMetadataGamut::Rec2020:
+			HDRMetadata.ColorSpace = BlackmagicDesign::EHDRMetaDataColorspace::Rec2020;
+			break;
+		}
+
+		HDRMetadata.EOTF = (BlackmagicDesign::EHDRMetaDataEOTF)HDROptions.EOTF;
+
+		return HDRMetadata;
+	}
+}
diff --git a/Engine/Plugins/Media/BlackmagicMedia/Source/BlackmagicMedia/Private/Player/BlackmagicMediaPlayer.cpp b/Engine/Plugins/Media/BlackmagicMedia/Source/BlackmagicMedia/Private/Player/BlackmagicMediaPlayer.cpp
--- a/Engine/Plugins/Media/BlackmagicMedia/Source/BlackmagicMedia/Private/Player/BlackmagicMediaPlayer.cpp
+++ b/Engine/Plugins/Media/BlackmagicMedia/Source/BlackmagicMedia/Private/Player/BlackmagicMediaPlayer.cpp
@@ -3,6 +3,7 @@
 #include "Player/BlackmagicMediaPlayer.h"
 
 #include "BlackmagicMediaPrivate.h"
+#include "BlackmagicMediaDefinitions.h"
 #include "BlackmagicMediaSource.h"
 
 #include "Engine/GameEngine.h"
@@ -15,6 +16,8 @@
 #include "MediaIOCoreEncodeTime.h"
 #include "MediaIOCoreFileWriter.h"
 #include "MediaIOCoreSamples.h"
+#include "MediaIOCoreTextureSampleBase.h"
+#include "MediaIOCoreUtilities.h"
 #include "Misc/App.h"
 #include "Misc/ScopeLock.h"
 #include "RenderCommandFence.h"
@@ -64,11 +67,10 @@ namespace BlackmagicMediaPlayerHelpers
 			, bReceivedValidFrame(false)
 			, bIsTimecodeExpected(false)
 			, bHasWarnedMissingTimecode(false)
-			, bIsSRGBInput(false)
 		{
 		}
 
-		bool Initialize(const BlackmagicDesign::FInputChannelOptions& InChannelInfo, bool bInEncodeTimecodeInTexel, int32 InMaxNumAudioFrameBuffer, int32 InMaxNumVideoFrameBuffer, bool bInIsSRGBInput)
+		bool Initialize(const BlackmagicDesign::FInputChannelOptions& InChannelInfo, bool bInEncodeTimecodeInTexel, int32 InMaxNumAudioFrameBuffer, int32 InMaxNumVideoFrameBuffer)
 		{
 			AddRef();
 
@@ -76,7 +78,6 @@ namespace BlackmagicMediaPlayerHelpers
 			MaxNumAudioFrameBuffer = InMaxNumAudioFrameBuffer;
 			MaxNumVideoFrameBuffer = InMaxNumVideoFrameBuffer;
 			bIsTimecodeExpected = InChannelInfo.TimecodeFormat != BlackmagicDesign::ETimecodeFormat::TCF_None;
-			bIsSRGBInput = bInIsSRGBInput;
 
 			BlackmagicDesign::ReferencePtr<BlackmagicDesign::IInputEventCallback> SelfRef(this);
 			BlackmagicIdendifier = BlackmagicDesign::RegisterCallbackForChannel(ChannelInfo, InChannelInfo, SelfRef);
@@ -292,6 +293,56 @@ namespace BlackmagicMediaPlayerHelpers
 						MediaIOCoreFileWriter::WriteRawFile(OutputFilename, reinterpret_cast<uint8*>(InFrameInfo.VideoBuffer), InFrameInfo.VideoPitch * InFrameInfo.VideoHeight);
 						bBlackmagicWriteOutputRawDataCmdEnable = false;
 					}
+					
+					UE::MediaIOCore::FColorFormatArgs ColorFormat;
+					FBlackmagicMediaHDROptions HDROptions = UE::BlackmagicMedia::MakeBlackmagicMediaHDROptions(InFrameInfo.HDRMetaData);
+		
+					if (MediaPlayer->bOverrideSourceEncoding)
+					{
+						ColorFormat.Encoding = (UE::Color::EEncoding) MediaPlayer->OverrideSourceEncoding;
+					}
+					else
+					{
+						const FTimespan TimeBetweenLogs = FTimespan::FromSeconds(5);
+
+						switch (HDROptions.EOTF)
+						{
+						case EBlackmagicHDRMetadataEOTF::SDR:
+							ColorFormat.Encoding = UE::Color::EEncoding::sRGB; // Missing support for BT.1886 so we default to sRGB instead.
+							break;
+						case EBlackmagicHDRMetadataEOTF::PQ:
+							ColorFormat.Encoding = UE::Color::EEncoding::ST2084;
+							break;
+						case EBlackmagicHDRMetadataEOTF::HLG:
+							UE_MEDIA_IO_LOG_THROTTLE(LogBlackmagicMedia, Warning, TimeBetweenLogs, TEXT("HLG EOTF is not supported at the moment. Defaulting to sRGB."));
+							ColorFormat.Encoding = UE::Color::EEncoding::sRGB;
+							break;
+						default:
+							checkNoEntry();
+							ColorFormat.Encoding = UE::Color::EEncoding::None;			
+						}
+					}
+		
+					if (MediaPlayer->bOverrideSourceColorSpace)
+					{
+						ColorFormat.ColorSpace = (UE::Color::EColorSpace) MediaPlayer->OverrideSourceColorSpace;
+					}
+					else
+					{
+						switch (HDROptions.Gamut)
+						{
+						case EBlackmagicHDRMetadataGamut::Rec709:
+							ColorFormat.ColorSpace = UE::Color::EColorSpace::sRGB;
+							break;
+						case EBlackmagicHDRMetadataGamut::Rec2020:
+							ColorFormat.ColorSpace = UE::Color::EColorSpace::Rec2020;
+							break;
+						default:
+							checkNoEntry();
+							ColorFormat.ColorSpace = UE::Color::EColorSpace::sRGB;
+							break;
+						}
+					}
 
 					if (bIsProgressivePicture)
 					{
@@ -320,7 +371,7 @@ namespace BlackmagicMediaPlayerHelpers
 						bool bInitializeResult = false;
 						if (bGPUDirectTexturesAvailable)
 						{
-							bInitializeResult = TextureSample->SetProperties(InFrameInfo.VideoPitch, InFrameInfo.VideoWidth, InFrameInfo.VideoHeight, SampleFormat, DecodedTime, MediaPlayer->VideoFrameRate, DecodedTimecode, bIsSRGBInput);
+							bInitializeResult = TextureSample->SetProperties(InFrameInfo.VideoPitch, InFrameInfo.VideoWidth, InFrameInfo.VideoHeight, SampleFormat, DecodedTime, MediaPlayer->VideoFrameRate, DecodedTimecode, ColorFormat);
 							if (bInitializeResult)
 							{
 								TextureSample->SetBuffer(InFrameInfo.VideoBuffer);
@@ -337,7 +388,7 @@ namespace BlackmagicMediaPlayerHelpers
 								, DecodedTime
 								, MediaPlayer->VideoFrameRate
 								, DecodedTimecode
-								, bIsSRGBInput);
+								, ColorFormat);
 						}
 
 						if (bInitializeResult)
@@ -366,7 +417,7 @@ namespace BlackmagicMediaPlayerHelpers
 							DecodedTime,
 							MediaPlayer->VideoFrameRate,
 							DecodedTimecode,
-							bIsSRGBInput ? UE::Color::EColorSpace::sRGB : UE::Color::EColorSpace::None
+							ColorFormat
 						};
 
 						const TArray<TSharedRef<FMediaIOCoreTextureSampleBase>> DeinterlacedSamples = MediaPlayer->Deinterlace(FrameInfo);
@@ -431,9 +482,6 @@ namespace BlackmagicMediaPlayerHelpers
 
 		bool bIsTimecodeExpected;
 		bool bHasWarnedMissingTimecode;
-
-		/** Whether this input is in sRGB space and needs a to linear conversion */
-		bool bIsSRGBInput;
 	};
 }
 
@@ -505,6 +553,11 @@ bool FBlackmagicMediaPlayer::Open(const FString& Url, const IMediaOptions* Optio
 	
 	const EMediaIOAutoDetectableTimecodeFormat TimecodeFormat = (EMediaIOAutoDetectableTimecodeFormat)(Options->GetMediaOption(BlackmagicMediaOption::TimecodeFormat, (int64)EMediaIOAutoDetectableTimecodeFormat::None));
 	const bool bAutoDetectTimecode = TimecodeFormat == EMediaIOAutoDetectableTimecodeFormat::Auto;
+
+	bOverrideSourceEncoding = Options->GetMediaOption(UE::CaptureCardMediaSource::OverrideSourceEncoding, true);
+	OverrideSourceEncoding = (ETextureSourceEncoding) Options->GetMediaOption(UE::CaptureCardMediaSource::SourceEncoding, (int64) ETextureSourceEncoding::TSE_Linear);
+	bOverrideSourceColorSpace = Options->GetMediaOption(UE::CaptureCardMediaSource::OverrideSourceColorSpace, true);
+	OverrideSourceColorSpace =  (ETextureColorSpace) Options->GetMediaOption(UE::CaptureCardMediaSource::SourceColorSpace, (int64) ETextureColorSpace::TCS_None);
 	
 	BlackmagicDesign::FChannelInfo ChannelInfo;
 	ChannelInfo.DeviceIndex = Options->GetMediaOption(BlackmagicMediaOption::DeviceIndex, (int64)0);
@@ -520,8 +573,7 @@ bool FBlackmagicMediaPlayer::Open(const FString& Url, const IMediaOptions* Optio
 	BlackmagicColorFormat = (EBlackmagicMediaSourceColorFormat)(Options->GetMediaOption(BlackmagicMediaOption::ColorFormat, (int64)EBlackmagicMediaSourceColorFormat::YUV8));
 
 	ChannelOptions.PixelFormat = BlackmagicColorFormat == EBlackmagicMediaSourceColorFormat::YUV8 ? BlackmagicDesign::EPixelFormat::pf_8Bits : BlackmagicDesign::EPixelFormat::pf_10Bits;
-	const bool bIsSRGBInput = Options->GetMediaOption(BlackmagicMediaOption::SRGBInput, true);
-
+	
 	switch (TimecodeFormat)
 	{
 	case EMediaIOAutoDetectableTimecodeFormat::Auto:
@@ -559,7 +611,7 @@ bool FBlackmagicMediaPlayer::Open(const FString& Url, const IMediaOptions* Optio
 	// Setup our different supported channels based on source settings
 	SetupSampleChannels();
 
-	bool bSuccess = EventCallback->Initialize(ChannelOptions, bEncodeTimecodeInTexel, MaxNumAudioFrameBuffer, MaxNumVideoFrameBuffer, bIsSRGBInput);
+	bool bSuccess = EventCallback->Initialize(ChannelOptions, bEncodeTimecodeInTexel, MaxNumAudioFrameBuffer, MaxNumVideoFrameBuffer);
 
 	if (!bSuccess)
 	{
diff --git a/Engine/Plugins/Media/BlackmagicMedia/Source/BlackmagicMedia/Private/Player/BlackmagicMediaPlayer.h b/Engine/Plugins/Media/BlackmagicMedia/Source/BlackmagicMedia/Private/Player/BlackmagicMediaPlayer.h
--- a/Engine/Plugins/Media/BlackmagicMedia/Source/BlackmagicMedia/Private/Player/BlackmagicMediaPlayer.h
+++ b/Engine/Plugins/Media/BlackmagicMedia/Source/BlackmagicMedia/Private/Player/BlackmagicMediaPlayer.h
@@ -23,8 +23,6 @@ namespace BlackmagicMediaPlayerHelpers
 
 class FBlackmagicMediaTextureSample : public FMediaIOCoreTextureSampleBase
 {
-public:
-	virtual const FMatrix& GetYUVToRGBMatrix() const override { return MediaShaders::YuvToRgbRec709Scaled; }
 };
 
 class FBlackmagicMediaAudioSamplePool : public TMediaObjectPool<FMediaIOCoreAudioSampleBase> { };
@@ -124,4 +122,16 @@ private:
 	EMediaIOSampleType SupportedSampleTypes;
 
 	EBlackmagicMediaSourceColorFormat BlackmagicColorFormat = EBlackmagicMediaSourceColorFormat::YUV8;
+
+	/** Whether to override the source encoding or to use the metadata embedded in the ancillary data of the signal. */
+	bool bOverrideSourceEncoding = true;
+
+	/** Encoding of the source texture. */
+	ETextureSourceEncoding OverrideSourceEncoding = ETextureSourceEncoding::TSE_Linear;
+
+	/** Whether to override the source color space or to use the metadata embedded in the ancillary data of the signal. */
+	bool bOverrideSourceColorSpace = true;
+
+	/** Color space of the source texture. */
+	ETextureColorSpace OverrideSourceColorSpace = ETextureColorSpace::TCS_None;
 };
diff --git a/Engine/Plugins/Media/BlackmagicMedia/Source/BlackmagicMedia/Public/BlackmagicMediaDefinitions.h b/Engine/Plugins/Media/BlackmagicMedia/Source/BlackmagicMedia/Public/BlackmagicMediaDefinitions.h
new file mode 100644
--- /dev/null
+++ b/Engine/Plugins/Media/BlackmagicMedia/Source/BlackmagicMedia/Public/BlackmagicMediaDefinitions.h
@@ -0,0 +1,60 @@
+// Copyright Epic Games, Inc. All Rights Reserved.
+
+#pragma once
+
+#include "UObject/Object.h"
+#include "BlackmagicMediaDefinitions.generated.h"
+
+namespace BlackmagicDesign
+{
+	struct FHDRMetaData;
+}
+
+/**
+ * HDR Transfer function.
+ * Must match Blackmagic::EBlackmagicHDRMetadataEOTF.
+ */
+UENUM()
+enum class EBlackmagicHDRMetadataEOTF : uint8
+{
+	SDR,
+	HDR,
+	PQ,
+	HLG,
+};
+	
+/**
+ * HDR Color Gamut.
+ * Must match Blackmagic::EBlackmagicHDRMetadataGamut.
+ */
+UENUM()
+enum class EBlackmagicHDRMetadataGamut : uint8
+{
+	Rec709,
+	Rec2020,
+};
+
+/**
+ * Set of metadata describing a HDR video signal.
+ */
+USTRUCT()
+struct BLACKMAGICMEDIA_API FBlackmagicMediaHDROptions
+{
+	GENERATED_BODY()
+
+	/** Transfer function to use for converting the video signal to an optical signal. */
+	UPROPERTY(EditAnywhere, Category = "HDR")
+	EBlackmagicHDRMetadataEOTF EOTF = EBlackmagicHDRMetadataEOTF::SDR;
+
+	/** The color gamut of the video signal. */
+	UPROPERTY(EditAnywhere, Category = "HDR")
+	EBlackmagicHDRMetadataGamut Gamut = EBlackmagicHDRMetadataGamut::Rec709;
+};
+
+namespace UE::BlackmagicMedia
+{
+	/** Create Blackmagic HDR UStruct from the BlackmagicLib's hdr struct. */
+	FBlackmagicMediaHDROptions MakeBlackmagicMediaHDROptions(const BlackmagicDesign::FHDRMetaData& HDROptions);
+	/** Create BlackmagicLib's hdr struct from the Blackmagic HDR UStruct. */
+	BlackmagicDesign::FHDRMetaData MakeBlackmagicHDROptions(const FBlackmagicMediaHDROptions& HDROptions);
+}
diff --git a/Engine/Plugins/Media/BlackmagicMedia/Source/BlackmagicMedia/Public/BlackmagicMediaSource.h b/Engine/Plugins/Media/BlackmagicMedia/Source/BlackmagicMedia/Public/BlackmagicMediaSource.h
--- a/Engine/Plugins/Media/BlackmagicMedia/Source/BlackmagicMedia/Public/BlackmagicMediaSource.h
+++ b/Engine/Plugins/Media/BlackmagicMedia/Source/BlackmagicMedia/Public/BlackmagicMediaSource.h
@@ -81,13 +81,15 @@ public:
 	UPROPERTY(BlueprintReadOnly, EditAnywhere, Category="Video", meta=(EditCondition="bCaptureVideo"))
 	EBlackmagicMediaSourceColorFormat ColorFormat;
 
+#if WITH_EDITORONLY_DATA
 	/**
 	 * Whether the video input is in sRGB color space.
 	 * A sRGB to Linear conversion will be applied resulting in a texture in linear space.
 	 * @Note If the texture is not in linear space, it won't look correct in the editor. Another pass will be required either through Composure or other means.
 	 */
-	UPROPERTY(BlueprintReadOnly, EditAnywhere, Category = "Video")
-	bool bIsSRGBInput;
+	UPROPERTY(meta=(DeprecatedProperty, DeprecationMessage = "Use ColorEncoding instead."))
+	bool bIsSRGBInput_DEPRECATED = false;
+#endif
 
 	/** Maximum number of video frames to buffer. */
 	UPROPERTY(BlueprintReadOnly, EditAnywhere, AdvancedDisplay, Category="Video", meta=(EditCondition="bCaptureVideo", ClampMin="1", ClampMax="32"))
diff --git a/Engine/Plugins/Media/BlackmagicMedia/Source/BlackmagicMediaOutput/Public/BlackmagicMediaOutput.h b/Engine/Plugins/Media/BlackmagicMedia/Source/BlackmagicMediaOutput/Public/BlackmagicMediaOutput.h
--- a/Engine/Plugins/Media/BlackmagicMedia/Source/BlackmagicMediaOutput/Public/BlackmagicMediaOutput.h
+++ b/Engine/Plugins/Media/BlackmagicMedia/Source/BlackmagicMediaOutput/Public/BlackmagicMediaOutput.h
@@ -3,6 +3,8 @@
 #pragma once
 
 #include "MediaOutput.h"
+
+#include "BlackmagicMediaDefinitions.h"
 #include "MediaIOCoreDefinitions.h"
 
 #include "BlackmagicMediaOutput.generated.h"
@@ -38,47 +40,6 @@ enum class EBlackmagicMediaOutputAudioBitDepth : uint8
 	Signed_32Bits = 32 UMETA(DisplayName = "32 bits signed")
 };
 
-/**
- * HDR Transfer function.
- * Must match BlackmagicDesign::EHDRMetaDataEOTF
- */
-UENUM()
-enum class EBlackmagicHDRMetadataEOTF : uint8
-{
-	SDR,
-	HDR,
-	PQ,
-	HLG
-};
-	
-/**
- * HDR Color Gamut.
- */
-UENUM()
-enum class EBlackmagicHDRMetadataGamut : uint8
-{
-	Rec709,
-	Rec2020
-};
-
-
-/**
- * Set of metadata describing a HDR video signal.
- */
-USTRUCT()
-struct BLACKMAGICMEDIAOUTPUT_API FBlackmagicMediaHDROptions
-{
-	GENERATED_BODY()
-
-	/** Transfer function to use for converting the video signal to an optical signal. */
-	UPROPERTY(EditAnywhere, Category = "HDR")
-	EBlackmagicHDRMetadataEOTF EOTF = EBlackmagicHDRMetadataEOTF::SDR;
-
-	/** The color gamut of the video signal. */
-	UPROPERTY(EditAnywhere, Category = "HDR")
-	EBlackmagicHDRMetadataGamut Gamut = EBlackmagicHDRMetadataGamut::Rec709;
-};
-
 /**
  * Output information for a MediaCapture.
  * @note	'Frame Buffer Pixel Format' must be set to at least 8 bits of alpha to enabled the Key.
diff --git a/Engine/Plugins/Media/MediaIOFramework/Source/MediaIOCore/MediaIOCore.Build.cs b/Engine/Plugins/Media/MediaIOFramework/Source/MediaIOCore/MediaIOCore.Build.cs
--- a/Engine/Plugins/Media/MediaIOFramework/Source/MediaIOCore/MediaIOCore.Build.cs
+++ b/Engine/Plugins/Media/MediaIOFramework/Source/MediaIOCore/MediaIOCore.Build.cs
@@ -13,6 +13,7 @@ namespace UnrealBuildTool.Rules
 				new string[]
 				{
 					"AudioExtensions",
+					"ColorManagement",
 					"Core",
 					"CoreUObject",
 					"Engine",
diff --git a/Engine/Plugins/Media/MediaIOFramework/Source/MediaIOCore/Private/CaptureCardMediaSource.cpp b/Engine/Plugins/Media/MediaIOFramework/Source/MediaIOCore/Private/CaptureCardMediaSource.cpp
--- a/Engine/Plugins/Media/MediaIOFramework/Source/MediaIOCore/Private/CaptureCardMediaSource.cpp
+++ b/Engine/Plugins/Media/MediaIOFramework/Source/MediaIOCore/Private/CaptureCardMediaSource.cpp
@@ -9,6 +9,16 @@ int64 UCaptureCardMediaSource::GetMediaOption(const FName& Key, int64 DefaultVal
 		return static_cast<int64>(InterlaceFieldOrder);
 	}
 
+	if (Key == UE::CaptureCardMediaSource::SourceColorSpace)
+	{
+		return (int64) OverrideSourceColorSpace;
+	}
+
+	if (Key == UE::CaptureCardMediaSource::SourceEncoding)
+	{
+		return (int64) OverrideSourceEncoding;
+	}
+
 	return UTimeSynchronizableMediaSource::GetMediaOption(Key, DefaultValue);
 }
 
@@ -25,13 +35,32 @@ FString UCaptureCardMediaSource::GetMediaOption(const FName& Key, const FString&
 	return UTimeSynchronizableMediaSource::GetMediaOption(Key, DefaultValue);
 }
 
+bool UCaptureCardMediaSource::GetMediaOption(const FName& Key, bool bDefaultValue) const
+{
+	if (Key == UE::CaptureCardMediaSource::OverrideSourceColorSpace)
+	{
+		return bOverrideSourceColorSpace;
+	}
+
+	if (Key == UE::CaptureCardMediaSource::OverrideSourceEncoding)
+	{
+		return bOverrideSourceEncoding;
+	}
+
+	return UTimeSynchronizableMediaSource::GetMediaOption(Key, bDefaultValue);
+}
+
 bool UCaptureCardMediaSource::HasMediaOption(const FName& Key) const
 {
 	if (Key == UE::CaptureCardMediaSource::Deinterlacer
-		|| Key == UE::CaptureCardMediaSource::InterlaceFieldOrder)
+		|| Key == UE::CaptureCardMediaSource::InterlaceFieldOrder
+		|| Key == UE::CaptureCardMediaSource::OverrideSourceEncoding
+		|| Key == UE::CaptureCardMediaSource::OverrideSourceColorSpace
+		|| Key == UE::CaptureCardMediaSource::SourceEncoding
+		|| Key == UE::CaptureCardMediaSource::SourceColorSpace)
 		{
 			return true;
 		}
 
 	return UTimeSynchronizableMediaSource::HasMediaOption(Key);
-}
\ No newline at end of file
+}
diff --git a/Engine/Plugins/Media/MediaIOFramework/Source/MediaIOCore/Private/MediaIOCoreDeinterlacer.cpp b/Engine/Plugins/Media/MediaIOFramework/Source/MediaIOCore/Private/MediaIOCoreDeinterlacer.cpp
--- a/Engine/Plugins/Media/MediaIOFramework/Source/MediaIOCore/Private/MediaIOCoreDeinterlacer.cpp
+++ b/Engine/Plugins/Media/MediaIOFramework/Source/MediaIOCore/Private/MediaIOCoreDeinterlacer.cpp
@@ -22,7 +22,7 @@ namespace UE::MediaIOCore
 			, InVideoFrame.Time
 			, InVideoFrame.FrameRate
 			, InVideoFrame.Timecode
-			, InVideoFrame.ColorSpace == Color::EColorSpace::sRGB ? true : false))
+			, InVideoFrame.ColorFormatArgs))
 		{
 			Samples.Add(TextureSample.ToSharedRef());
 		}
@@ -70,7 +70,7 @@ namespace UE::MediaIOCore
 				, InVideoFrame.Time
 				, InVideoFrame.FrameRate
 				, InVideoFrame.Timecode
-				, InVideoFrame.ColorSpace == Color::EColorSpace::sRGB ? true : false))
+				, InVideoFrame.ColorFormatArgs))
 			{
 				Samples.Add(TextureSampleEven.ToSharedRef());
 			}
@@ -86,7 +86,7 @@ namespace UE::MediaIOCore
 				, InVideoFrame.Time
 				, InVideoFrame.FrameRate
 				, InVideoFrame.Timecode
-				, InVideoFrame.ColorSpace == Color::EColorSpace::sRGB ? true : false))
+				, InVideoFrame.ColorFormatArgs))
 			{
 				Samples.Add(TextureSampleOdd.ToSharedRef());
 			}
@@ -127,7 +127,7 @@ namespace UE::MediaIOCore
 				, InVideoFrame.Time
 				, InVideoFrame.FrameRate
 				, InVideoFrame.Timecode
-				, InVideoFrame.ColorSpace == Color::EColorSpace::sRGB ? true : false))
+				, InVideoFrame.ColorFormatArgs))
 			{
 				Samples.Add(TextureSample.ToSharedRef());
 			}
@@ -159,7 +159,7 @@ namespace UE::MediaIOCore
 				, InVideoFrame.Time
 				, InVideoFrame.FrameRate
 				, InVideoFrame.Timecode
-				, InVideoFrame.ColorSpace == Color::EColorSpace::sRGB ? true : false))
+				, InVideoFrame.ColorFormatArgs))
 			{
 				Samples.Add(TextureSample.ToSharedRef());
 			}
diff --git a/Engine/Plugins/Media/MediaIOFramework/Source/MediaIOCore/Private/Shared/MediaIOCoreTextureSampleBase.cpp b/Engine/Plugins/Media/MediaIOFramework/Source/MediaIOCore/Private/Shared/MediaIOCoreTextureSampleBase.cpp
--- a/Engine/Plugins/Media/MediaIOFramework/Source/MediaIOCore/Private/Shared/MediaIOCoreTextureSampleBase.cpp
+++ b/Engine/Plugins/Media/MediaIOFramework/Source/MediaIOCore/Private/Shared/MediaIOCoreTextureSampleBase.cpp
@@ -12,16 +12,15 @@ FMediaIOCoreTextureSampleBase::FMediaIOCoreTextureSampleBase()
 	, Stride(0)
 	, Width(0)
 	, Height(0)
-	, bIsSRGBInput(false)
 {
 }
 
 
-bool FMediaIOCoreTextureSampleBase::Initialize(const void* InVideoBuffer, uint32 InBufferSize, uint32 InStride, uint32 InWidth, uint32 InHeight, EMediaTextureSampleFormat InSampleFormat, FTimespan InTime, const FFrameRate& InFrameRate, const TOptional<FTimecode>& InTimecode, bool bInIsSRGBInput)
+bool FMediaIOCoreTextureSampleBase::Initialize(const void* InVideoBuffer, uint32 InBufferSize, uint32 InStride, uint32 InWidth, uint32 InHeight, EMediaTextureSampleFormat InSampleFormat, FTimespan InTime, const FFrameRate& InFrameRate, const TOptional<FTimecode>& InTimecode, const UE::MediaIOCore::FColorFormatArgs& InColorFormatArgs)
 {
 	FreeSample();
 
-	if (!SetProperties(InStride, InWidth, InHeight, InSampleFormat, InTime, InFrameRate, InTimecode, bInIsSRGBInput))
+	if (!SetProperties(InStride, InWidth, InHeight, InSampleFormat, InTime, InFrameRate, InTimecode, InColorFormatArgs))
 	{
 		return false;
 	}
@@ -30,11 +29,11 @@ bool FMediaIOCoreTextureSampleBase::Initialize(const void* InVideoBuffer, uint32
 }
 
 
-bool FMediaIOCoreTextureSampleBase::Initialize(const TArray<uint8>& InVideoBuffer, uint32 InStride, uint32 InWidth, uint32 InHeight, EMediaTextureSampleFormat InSampleFormat, FTimespan InTime, const FFrameRate& InFrameRate, const TOptional<FTimecode>& InTimecode, bool bInIsSRGBInput)
+bool FMediaIOCoreTextureSampleBase::Initialize(const TArray<uint8>& InVideoBuffer, uint32 InStride, uint32 InWidth, uint32 InHeight, EMediaTextureSampleFormat InSampleFormat, FTimespan InTime, const FFrameRate& InFrameRate, const TOptional<FTimecode>& InTimecode, const UE::MediaIOCore::FColorFormatArgs& InColorFormatArgs)
 {
 	FreeSample();
 
-	if (!SetProperties(InStride, InWidth, InHeight, InSampleFormat, InTime, InFrameRate, InTimecode, bInIsSRGBInput))
+	if (!SetProperties(InStride, InWidth, InHeight, InSampleFormat, InTime, InFrameRate, InTimecode, InColorFormatArgs))
 	{
 		return false;
 	}
@@ -43,11 +42,11 @@ bool FMediaIOCoreTextureSampleBase::Initialize(const TArray<uint8>& InVideoBuffe
 }
 
 
-bool FMediaIOCoreTextureSampleBase::Initialize(TArray<uint8>&& InVideoBuffer, uint32 InStride, uint32 InWidth, uint32 InHeight, EMediaTextureSampleFormat InSampleFormat, FTimespan InTime, const FFrameRate& InFrameRate, const TOptional<FTimecode>& InTimecode, bool bInIsSRGBInput)
+bool FMediaIOCoreTextureSampleBase::Initialize(TArray<uint8>&& InVideoBuffer, uint32 InStride, uint32 InWidth, uint32 InHeight, EMediaTextureSampleFormat InSampleFormat, FTimespan InTime, const FFrameRate& InFrameRate, const TOptional<FTimecode>& InTimecode, const UE::MediaIOCore::FColorFormatArgs& InColorFormatArgs)
 {
 	FreeSample();
 
-	if (!SetProperties(InStride, InWidth, InHeight, InSampleFormat, InTime, InFrameRate, InTimecode, bInIsSRGBInput))
+	if (!SetProperties(InStride, InWidth, InHeight, InSampleFormat, InTime, InFrameRate, InTimecode, InColorFormatArgs))
 	{
 		return false;
 	}
@@ -94,7 +93,7 @@ bool FMediaIOCoreTextureSampleBase::SetBuffer(TArray<uint8>&& InVideoBuffer)
 	return true;
 }
 
-bool FMediaIOCoreTextureSampleBase::SetProperties(uint32 InStride, uint32 InWidth, uint32 InHeight, EMediaTextureSampleFormat InSampleFormat, FTimespan InTime, const FFrameRate& InFrameRate, const TOptional<FTimecode>& InTimecode, bool bInIsSRGBInput)
+bool FMediaIOCoreTextureSampleBase::SetProperties(uint32 InStride, uint32 InWidth, uint32 InHeight, EMediaTextureSampleFormat InSampleFormat, FTimespan InTime, const FFrameRate& InFrameRate, const TOptional<FTimecode>& InTimecode, const UE::MediaIOCore::FColorFormatArgs& InColorFormatArgs)
 {
 	if (InSampleFormat == EMediaTextureSampleFormat::Undefined)
 	{
@@ -108,16 +107,18 @@ bool FMediaIOCoreTextureSampleBase::SetProperties(uint32 InStride, uint32 InWidt
 	Time = InTime;
 	Duration = FTimespan(ETimespan::TicksPerSecond * InFrameRate.AsInterval());
 	Timecode = InTimecode;
-	bIsSRGBInput = bInIsSRGBInput;
+	Encoding = InColorFormatArgs.Encoding;
+	ColorSpace = InColorFormatArgs.ColorSpace;
+	ColorSpaceStruct = UE::Color::FColorSpace(ColorSpace);
 
 	return true;
 }
 
-bool FMediaIOCoreTextureSampleBase::InitializeWithEvenOddLine(bool bUseEvenLine, const void* InVideoBuffer, uint32 InBufferSize, uint32 InStride, uint32 InWidth, uint32 InHeight, EMediaTextureSampleFormat InSampleFormat, FTimespan InTime, const FFrameRate& InFrameRate, const TOptional<FTimecode>& InTimecode, bool bInIsSRGBInput)
+bool FMediaIOCoreTextureSampleBase::InitializeWithEvenOddLine(bool bUseEvenLine, const void* InVideoBuffer, uint32 InBufferSize, uint32 InStride, uint32 InWidth, uint32 InHeight, EMediaTextureSampleFormat InSampleFormat, FTimespan InTime, const FFrameRate& InFrameRate, const TOptional<FTimecode>& InTimecode, const UE::MediaIOCore::FColorFormatArgs& InColorFormatArgs)
 {
 	FreeSample();
 
-	if (!SetProperties(InStride, InWidth, InHeight/2, InSampleFormat, InTime, InFrameRate, InTimecode, bInIsSRGBInput))
+	if (!SetProperties(InStride, InWidth, InHeight/2, InSampleFormat, InTime, InFrameRate, InTimecode, InColorFormatArgs))
 	{
 		return false;
 	}
@@ -171,3 +172,56 @@ void FMediaIOCoreTextureSampleBase::ShutdownPoolable()
 
 	FreeSample();
 }
+
+const FMatrix& FMediaIOCoreTextureSampleBase::GetYUVToRGBMatrix() const
+{
+	switch (ColorSpace)
+	{
+	case UE::Color::EColorSpace::sRGB:
+		return MediaShaders::YuvToRgbRec709Scaled;
+	case UE::Color::EColorSpace::Rec2020:
+		return MediaShaders::YuvToRgbRec2020Scaled;
+	default:
+		return MediaShaders::YuvToRgbRec709Scaled;
+	}
+}
+
+bool FMediaIOCoreTextureSampleBase::IsOutputSrgb() const
+{
+	return Encoding == UE::Color::EEncoding::sRGB;
+}
+
+FMatrix44f FMediaIOCoreTextureSampleBase::GetGamutToXYZMatrix() const
+{
+	return FMatrix44f(ColorSpaceStruct.GetRgbToXYZ().GetTransposed());
+}
+
+FVector2f FMediaIOCoreTextureSampleBase::GetWhitePoint() const
+{
+	return FVector2f(ColorSpaceStruct.GetWhiteChromaticity());
+}
+
+FVector2f FMediaIOCoreTextureSampleBase::GetDisplayPrimaryRed() const
+{
+	return FVector2f(ColorSpaceStruct.GetRedChromaticity());
+}
+
+FVector2f FMediaIOCoreTextureSampleBase::GetDisplayPrimaryGreen() const
+{
+	return FVector2f(ColorSpaceStruct.GetGreenChromaticity());
+}
+
+FVector2f FMediaIOCoreTextureSampleBase::GetDisplayPrimaryBlue() const
+{
+	return FVector2f(ColorSpaceStruct.GetBlueChromaticity());
+}
+
+UE::Color::EEncoding FMediaIOCoreTextureSampleBase::GetEncodingType() const
+{
+	return Encoding;
+}
+
+float FMediaIOCoreTextureSampleBase::GetHDRNitsNormalizationFactor() const
+{
+	return (GetEncodingType() == UE::Color::EEncoding::sRGB || GetEncodingType() == UE::Color::EEncoding::Linear) ? 1.0f : kMediaSample_HDR_NitsNormalizationFactor;
+}
diff --git a/Engine/Plugins/Media/MediaIOFramework/Source/MediaIOCore/Private/Tests/MediaIODeinterlacerTests.cpp b/Engine/Plugins/Media/MediaIOFramework/Source/MediaIOCore/Private/Tests/MediaIODeinterlacerTests.cpp
--- a/Engine/Plugins/Media/MediaIOFramework/Source/MediaIOCore/Private/Tests/MediaIODeinterlacerTests.cpp
+++ b/Engine/Plugins/Media/MediaIOFramework/Source/MediaIOCore/Private/Tests/MediaIODeinterlacerTests.cpp
@@ -57,7 +57,7 @@ namespace UE::MediaIOTests
 			FTimespan::FromSeconds(0),
 			FFrameRate(),
 			FTimecode(),
-			UE::Color::EColorSpace::None
+			UE::MediaIOCore::FColorFormatArgs{UE::Color::EEncoding::Linear, UE::Color::EColorSpace::sRGB}
 		};
 
 		// Execute test
@@ -98,7 +98,7 @@ namespace UE::MediaIOTests
 			FTimespan::FromSeconds(0),
 			FFrameRate(),
 			FTimecode(),
-			UE::Color::EColorSpace::None
+			UE::MediaIOCore::FColorFormatArgs{UE::Color::EEncoding::Linear, UE::Color::EColorSpace::sRGB}
 		};
 
 		// Execute test
@@ -134,7 +134,7 @@ namespace UE::MediaIOTests
 			FTimespan::FromSeconds(0),
 			FFrameRate(),
 			FTimecode(),
-			UE::Color::EColorSpace::None
+			UE::MediaIOCore::FColorFormatArgs{UE::Color::EEncoding::Linear, UE::Color::EColorSpace::sRGB}
 		};
 
 		// Execute test
@@ -171,7 +171,7 @@ namespace UE::MediaIOTests
 			FTimespan::FromSeconds(0),
 			FFrameRate(),
 			FTimecode(),
-			UE::Color::EColorSpace::None
+			UE::MediaIOCore::FColorFormatArgs{UE::Color::EEncoding::Linear, UE::Color::EColorSpace::sRGB}
 		};
 
 		// Execute test
diff --git a/Engine/Plugins/Media/MediaIOFramework/Source/MediaIOCore/Public/CaptureCardMediaSource.h b/Engine/Plugins/Media/MediaIOFramework/Source/MediaIOCore/Public/CaptureCardMediaSource.h
--- a/Engine/Plugins/Media/MediaIOFramework/Source/MediaIOCore/Public/CaptureCardMediaSource.h
+++ b/Engine/Plugins/Media/MediaIOFramework/Source/MediaIOCore/Public/CaptureCardMediaSource.h
@@ -4,6 +4,7 @@
 
 #include "TimeSynchronizableMediaSource.h"
 #include "Containers/UnrealString.h"
+#include "Engine/TextureDefines.h"
 #include "MediaIOCoreDefinitions.h"
 #include "MediaIOCoreDeinterlacer.h"
 #include "HAL/Platform.h"
@@ -18,8 +19,25 @@ namespace UE::CaptureCardMediaSource
 {
 	static FName Deinterlacer("Deinterlacer");
 	static FName InterlaceFieldOrder("InterlaceFieldOrder");
+	static FName OverrideSourceEncoding("OverrideSourceEncoding");
+	static FName SourceEncoding("SourceEncoding");
+	static FName OverrideSourceColorSpace("OverrideSourceColorSpace");
+	static FName SourceColorSpace("SourceColorcSpace");
 }
 
+
+/** List of texture source encodings that can be converted to linear. (Integer values match the ETextureSourceEncoding values in TextureDefines.h */
+UENUM()
+enum class EMediaIOCoreSourceEncoding : uint8
+{
+	Linear		= 1 UMETA(DisplayName = "Linear"),
+	sRGB		= 2 UMETA(DisplayName = "sRGB"),
+	ST2084		= 3 UMETA(DisplayName = "ST 2084/PQ"),
+	//BT1886		= 5 UMETA(DisplayName = "BT1886/Gamma 2.4"),
+	SLog3		= 12 UMETA(DisplayName = "SLog3"),
+	MAX,
+};
+
 /**
  * Base class for media sources that are coming from a capture card.
  */
@@ -47,10 +65,27 @@ public:
 	UPROPERTY(BlueprintReadOnly, EditAnywhere, Category = "Video")
 	EMediaIOInterlaceFieldOrder InterlaceFieldOrder = EMediaIOInterlaceFieldOrder::TopFieldFirst;
 
+	/** Whether to override the source encoding or to use the metadata embedded in the ancillary data of the signal. */
+	UPROPERTY(EditAnywhere, Category = "Video", meta = (InlineEditConditionToggle))
+	bool bOverrideSourceEncoding = true;
+
+	/** Encoding of the source texture. */
+	UPROPERTY(BlueprintReadOnly, EditAnywhere, Category = "Video", meta = (EditCondition = "bOverrideSourceEncoding"))
+	EMediaIOCoreSourceEncoding OverrideSourceEncoding = EMediaIOCoreSourceEncoding::Linear;
+
+	/** Whether to override the source color space or to use the metadata embedded in the ancillary data of the signal. */
+	UPROPERTY(EditAnywhere, Category = "Video", meta = (InlineEditConditionToggle))
+	bool bOverrideSourceColorSpace = true;
+
+	/** Color space of the source texture. */
+	UPROPERTY(BlueprintReadOnly, EditAnywhere, Category = "Video", meta = (EditCondition = "bOverrideSourceColorSpace"))
+	ETextureColorSpace OverrideSourceColorSpace = ETextureColorSpace::TCS_sRGB;
+
 public:
 	//~ IMediaOptions interface
 	using Super::GetMediaOption;
 	virtual FString GetMediaOption(const FName& Key, const FString& DefaultValue) const override;
 	virtual int64 GetMediaOption(const FName& Key, int64 DefaultValue) const override;
+	virtual bool GetMediaOption(const FName& Key, bool bDefaultValue) const override;
 	virtual bool HasMediaOption(const FName& Key) const override;
 };
diff --git a/Engine/Plugins/Media/MediaIOFramework/Source/MediaIOCore/Public/MediaIOCoreDeinterlacer.h b/Engine/Plugins/Media/MediaIOFramework/Source/MediaIOCore/Public/MediaIOCoreDeinterlacer.h
--- a/Engine/Plugins/Media/MediaIOFramework/Source/MediaIOCore/Public/MediaIOCoreDeinterlacer.h
+++ b/Engine/Plugins/Media/MediaIOFramework/Source/MediaIOCore/Public/MediaIOCoreDeinterlacer.h
@@ -23,7 +23,7 @@ namespace UE::MediaIOCore
 		, FTimespan InTime
 		, const FFrameRate& InFrameRate
 		, const TOptional<FTimecode>& InTimecode
-		, Color::EColorSpace InColorSpace)
+		, const UE::MediaIOCore::FColorFormatArgs& InColorFormatArgs)
 			: VideoBuffer(InVideoBuffer)
 			, BufferSize(InBufferSize)
 			, Stride(InStride)
@@ -33,7 +33,7 @@ namespace UE::MediaIOCore
 			, Time(InTime)
 			, FrameRate(InFrameRate)
 			, Timecode(InTimecode)
-			, ColorSpace(InColorSpace)
+			, ColorFormatArgs(InColorFormatArgs)
 		{
 		}
 
@@ -46,7 +46,7 @@ namespace UE::MediaIOCore
 		FTimespan Time;
 		const FFrameRate& FrameRate;
 		const TOptional<FTimecode>& Timecode;
-		Color::EColorSpace ColorSpace;
+		const UE::MediaIOCore::FColorFormatArgs& ColorFormatArgs;
 	};
 
 	/**
diff --git a/Engine/Plugins/Media/MediaIOFramework/Source/MediaIOCore/Public/MediaIOCoreTextureSampleBase.h b/Engine/Plugins/Media/MediaIOFramework/Source/MediaIOCore/Public/MediaIOCoreTextureSampleBase.h
--- a/Engine/Plugins/Media/MediaIOFramework/Source/MediaIOCore/Public/MediaIOCoreTextureSampleBase.h
+++ b/Engine/Plugins/Media/MediaIOFramework/Source/MediaIOCore/Public/MediaIOCoreTextureSampleBase.h
@@ -3,13 +3,52 @@
 #pragma once
 
 #include "IMediaTextureSample.h"
-#include "MediaIOCoreDefinitions.h"
 #include "MediaObjectPool.h"
+
+
+#include "ColorSpace.h"
 #include "Misc/FrameRate.h"
 #include "Templates/RefCounting.h"
 
 class FRHITexture;
 
+namespace UE::MediaIOCore
+{
+	struct FColorFormatArgs
+	{
+		FColorFormatArgs() = default;
+		
+		/** Bool constructor to allow backwards compatibility with previous method definitions of FMediaIOCoreTextureSampleBase::Initialize. */
+		FColorFormatArgs(bool bIsSRGBInput)
+		{
+			if (bIsSRGBInput)
+			{
+				Encoding = UE::Color::EEncoding::sRGB;
+				ColorSpace = UE::Color::EColorSpace::sRGB;
+			}
+			else
+			{
+				Encoding = UE::Color::EEncoding::Linear;
+				ColorSpace = UE::Color::EColorSpace::sRGB;
+			}
+		}
+
+		FColorFormatArgs(UE::Color::EEncoding InEncoding, UE::Color::EColorSpace InColorSpace)
+			: Encoding(InEncoding)
+			, ColorSpace(InColorSpace)
+		{
+		}
+		
+
+		/** Encoding of the texture. */
+		UE::Color::EEncoding Encoding = UE::Color::EEncoding::Linear;
+
+		/** Color space of the texture. */
+		UE::Color::EColorSpace ColorSpace = UE::Color::EColorSpace::sRGB;
+	};
+}
+
+
 /**
  * Implements the IMediaTextureSample/IMediaPoolable interface.
  */
@@ -33,9 +72,9 @@ public:
 	 * @param InTime The sample time (in the player's own clock).
 	 * @param InFrameRate The framerate of the media that produce the sample.
 	 * @param InTimecode The sample timecode if available.
-	 * @param bInIsSRGBInput The sample source is in sSRGB so apply a linear conversion.
+	 * @param InColorFormatArgs Information about the texture color encoding and color space.
 	 */
-	bool Initialize(const void* InVideoBuffer, uint32 InBufferSize, uint32 InStride, uint32 InWidth, uint32 InHeight, EMediaTextureSampleFormat InSampleFormat, FTimespan InTime, const FFrameRate& InFrameRate, const TOptional<FTimecode>& InTimecode, bool bInIsSRGBInput);
+	bool Initialize(const void* InVideoBuffer, uint32 InBufferSize, uint32 InStride, uint32 InWidth, uint32 InHeight, EMediaTextureSampleFormat InSampleFormat, FTimespan InTime, const FFrameRate& InFrameRate, const TOptional<FTimecode>& InTimecode, const UE::MediaIOCore::FColorFormatArgs& InColorFormatArgs);
 
 	/**
 	 * Initialize the sample.
@@ -48,9 +87,9 @@ public:
 	 * @param InTime The sample time (in the player's own clock).
 	 * @param InFrameRate The framerate of the media that produce the sample.
 	 * @param InTimecode The sample timecode if available.
-	 * @param bInIsSRGBInput The sample source is in sSRGB so apply a linear conversion.
+	 * @param InColorFormatArgs Information about the texture color encoding and color space.
 	 */
-	bool Initialize(const TArray<uint8>& InVideoBuffer, uint32 InStride, uint32 InWidth, uint32 InHeight, EMediaTextureSampleFormat InSampleFormat, FTimespan InTime, const FFrameRate& InFrameRate, const TOptional<FTimecode>& InTimecode, bool bInIsSRGBInput);
+	bool Initialize(const TArray<uint8>& InVideoBuffer, uint32 InStride, uint32 InWidth, uint32 InHeight, EMediaTextureSampleFormat InSampleFormat, FTimespan InTime, const FFrameRate& InFrameRate, const TOptional<FTimecode>& InTimecode, const UE::MediaIOCore::FColorFormatArgs& InColorFormatArgs);
 
 	/**
 	 * Initialize the sample.
@@ -63,9 +102,9 @@ public:
 	 * @param InTime The sample time (in the player's own clock).
 	 * @param InFrameRate The framerate of the media that produce the sample.
 	 * @param InTimecode The sample timecode if available.
-	 * @param bInIsSRGBInput The sample source is in sSRGB so apply a linear conversion.
+	 * @param InColorFormatArgs Information about the texture color encoding and color space.
 	 */
-	bool Initialize(TArray<uint8>&& InVideoBuffer, uint32 InStride, uint32 InWidth, uint32 InHeight, EMediaTextureSampleFormat InSampleFormat, FTimespan InTime, const FFrameRate& InFrameRate, const TOptional<FTimecode>& InTimecode, bool bInIsSRGBInput);
+	bool Initialize(TArray<uint8>&& InVideoBuffer, uint32 InStride, uint32 InWidth, uint32 InHeight, EMediaTextureSampleFormat InSampleFormat, FTimespan InTime, const FFrameRate& InFrameRate, const TOptional<FTimecode>& InTimecode, const UE::MediaIOCore::FColorFormatArgs& InColorFormatArgs);
 
 	/**
 	 * Initialize the sample.
@@ -99,9 +138,9 @@ public:
 	 * @param InTime The sample time (in the player's own clock).
 	 * @param InFrameRate The framerate of the media that produce the sample.
 	 * @param InTimecode The sample timecode if available.
-	 * @param bInIsSRGBInput The sample source is in sSRGB so apply a linear conversion.
+	 * @param InColorFormatArgs Information about the texture color encoding and color space.
 	 */
-	bool SetProperties(uint32 InStride, uint32 InWidth, uint32 InHeight, EMediaTextureSampleFormat InSampleFormat, FTimespan InTime, const FFrameRate& InFrameRate, const TOptional<FTimecode>& InTimecode, bool bInIsSRGBInput);
+	bool SetProperties(uint32 InStride, uint32 InWidth, uint32 InHeight, EMediaTextureSampleFormat InSampleFormat, FTimespan InTime, const FFrameRate& InFrameRate, const TOptional<FTimecode>& InTimecode, const UE::MediaIOCore::FColorFormatArgs& InColorFormatArgs);
 
 	/**
 	 * Initialize the sample with half it's original height and take only the odd or even line.
@@ -114,9 +153,9 @@ public:
 	 * @param InSampleFormat The sample format of the video buffer.
 	 * @param InTime The sample time (in the player's own clock).
 	 * @param InTimecode The sample timecode if available.
-	 * @param bInIsSRGBInput The sample source is in sSRGB so apply a linear conversion.
+	 * @param InColorFormatArgs Information about the texture color encoding and color space.
 	 */
-	bool InitializeWithEvenOddLine(bool bUseEvenLine, const void* InVideoBuffer, uint32 InBufferSize, uint32 InStride, uint32 InWidth, uint32 InHeight, EMediaTextureSampleFormat InSampleFormat, FTimespan InTime, const FFrameRate& InFrameRate, const TOptional<FTimecode>& InTimecode, bool bInIsSRGBInput);
+	bool InitializeWithEvenOddLine(bool bUseEvenLine, const void* InVideoBuffer, uint32 InBufferSize, uint32 InStride, uint32 InWidth, uint32 InHeight, EMediaTextureSampleFormat InSampleFormat, FTimespan InTime, const FFrameRate& InFrameRate, const TOptional<FTimecode>& InTimecode, const UE::MediaIOCore::FColorFormatArgs& InColorFormatArgs);
 
 	/**
 	 * Set the sample buffer with half it's original height and take only the odd or even line.
@@ -198,10 +237,15 @@ public:
 		return true;
 	}
 
-	virtual bool IsOutputSrgb() const override
-	{
-		return bIsSRGBInput;
-	}
+	virtual bool IsOutputSrgb() const override;
+	
+	virtual FMatrix44f GetGamutToXYZMatrix() const override;
+	virtual FVector2f GetWhitePoint() const override;
+	virtual FVector2f GetDisplayPrimaryRed() const override;
+	virtual FVector2f GetDisplayPrimaryGreen() const override;
+	virtual FVector2f GetDisplayPrimaryBlue() const override;
+	virtual UE::Color::EEncoding GetEncodingType() const override;
+	virtual float GetHDRNitsNormalizationFactor() const override;
 
 	virtual const void* GetBuffer() override
 	{
@@ -260,6 +304,13 @@ protected:
 		Buffer.Reset();
 	}
 
+	/**
+	 * Get YUV to RGB conversion matrix
+	 *
+	 * @return MediaIOCore Yuv To Rgb matrix
+	 */
+	virtual const FMatrix& GetYUVToRGBMatrix() const override;
+
 protected:
 	/** Duration for which the sample is valid. */
 	FTimespan Duration;
@@ -282,6 +333,18 @@ protected:
 	TArray<uint8, TAlignedHeapAllocator<4096>> Buffer;
 	void* ExternalBuffer = nullptr;
 
-	/** Wheter the sample is in sRGB space and requires an explicit conversion to linear */
-	bool bIsSRGBInput;
+#if WITH_EDITORONLY_DATA
+	UE_DEPRECATED(5.3, "Please use Encoding instead.")
+	/** Whether the sample is in sRGB space and requires an explicit conversion to linear */
+	bool bIsSRGBInput = false;
+#endif
+
+	/** Color encoding of the incoming texture. */
+	UE::Color::EEncoding Encoding = UE::Color::EEncoding::Linear;
+
+	/** Color space enum of the incoming texture. */
+	UE::Color::EColorSpace ColorSpace = UE::Color::EColorSpace::sRGB;
+
+	/** Color space structure of the incoming texture. Used for retrieving chromaticities. */
+	UE::Color::FColorSpace ColorSpaceStruct = UE::Color::FColorSpace(UE::Color::EColorSpace::sRGB);
 };
diff --git a/Engine/Plugins/Media/MediaIOFramework/Source/MediaIOCore/Public/MediaIOCoreUtilities.h b/Engine/Plugins/Media/MediaIOFramework/Source/MediaIOCore/Public/MediaIOCoreUtilities.h
new file mode 100644
--- /dev/null
+++ b/Engine/Plugins/Media/MediaIOFramework/Source/MediaIOCore/Public/MediaIOCoreUtilities.h
@@ -0,0 +1,69 @@
+// Copyright Epic Games, Inc. All Rights Reserved.
+#pragma once
+
+#include "CoreGlobals.h"
+#include "Containers/UnrealString.h"
+#include "Containers/Map.h"
+#include "Logging/LogCategory.h"
+#include "Logging/LogMacros.h"
+#include "Logging/LogVerbosity.h"
+#include "HAL/PlatformTime.h"
+#include "HAL/CriticalSection.h"
+#include "Misc/Timespan.h"
+
+namespace UE::MediaIO
+{
+#if NO_LOGGING
+	void LogThrottle(const FNoLoggingCategory& InLogCategory, ELogVerbosity::Type InVerbosity, const FTimespan& InTimeBetweenLogs, const FString& LogDetails, const FString& FileName, int32 LineNumber);
+#else
+	void LogThrottle(const FLogCategoryBase& InLogCategory, ELogVerbosity::Type InVerbosity, const FTimespan& InTimeBetweenLogs, const FString& LogDetails, const FString& FileName, int32 LineNumber);
+#endif
+}
+
+#define UE_MEDIA_IO_LOG_THROTTLE(LogCategory, Verbosity, TimeBetweenLogs, Format, ...) (UE::MediaIO::LogThrottle(LogCategory, ELogVerbosity::Verbosity, TimeBetweenLogs, *FString::Printf(Format, ##__VA_ARGS__), FString(__FILE__), __LINE__))
+
+namespace UE::MediaIO
+{
+#if NO_LOGGING
+	void LogThrottle(const FNoLoggingCategory& InLogCategory, ELogVerbosity::Type InVerbosity, const FTimespan& InTimeBetweenLogs, const FString& LogDetails, const FString& FileName, int32 LineNumber)
+	{
+	}
+#else
+	void LogThrottle(const FLogCategoryBase& InLogCategory, ELogVerbosity::Type InVerbosity, const FTimespan& InTimeBetweenLogs, const FString& LogDetails, const FString& FileName, int32 LineNumber)
+	{
+		// This code was adapted from AudioMixer.h to support throttling instead of only logging once.
+		// Log once to avoid Spam.
+		static FCriticalSection Cs;
+		static TMap<uint32, double> LogHistory;
+
+		FScopeLock Lock(&Cs);
+		const FString MessageToHash = FString::Printf(TEXT("%s (File %s, Line %d)"), *LogDetails, *FileName, LineNumber);
+
+		uint32 Hash = GetTypeHash(MessageToHash);
+
+		bool bShouldLog = false;
+		double* LastLogTime = LogHistory.Find(Hash);
+		
+		double NewTime = FPlatformTime::Seconds();
+		if (LastLogTime)
+		{
+			const double TimeDiff = NewTime - *LastLogTime;
+			if (TimeDiff > InTimeBetweenLogs.GetTotalSeconds())
+			{
+				*LastLogTime = NewTime;
+				bShouldLog = true;
+			}
+		}
+		else
+		{
+			bShouldLog = true;
+			LogHistory.Add(Hash, NewTime);
+		}
+
+		if (bShouldLog)
+		{
+			GLog->Log(InLogCategory.GetCategoryName(), InVerbosity, LogDetails);
+		}
+	}
+#endif
+}
