Fix deprecation error in MediaIOCoreTextureSampleBase.h

diff --git a/Engine/Plugins/Media/MediaIOFramework/Source/MediaIOCore/Public/MediaIOCoreTextureSampleBase.h b/Engine/Plugins/Media/MediaIOFramework/Source/MediaIOCore/Public/MediaIOCoreTextureSampleBase.h
--- a/Engine/Plugins/Media/MediaIOFramework/Source/MediaIOCore/Public/MediaIOCoreTextureSampleBase.h
+++ b/Engine/Plugins/Media/MediaIOFramework/Source/MediaIOCore/Public/MediaIOCoreTextureSampleBase.h
@@ -60,6 +60,14 @@ class MEDIAIOCORE_API FMediaIOCoreTextureSampleBase
 public:
 	FMediaIOCoreTextureSampleBase();
 
+	// Note: We need to explicitly disable warnings on these constructors/operators for clang to be happy with deprecated variables
+	PRAGMA_DISABLE_DEPRECATION_WARNINGS
+	FMediaIOCoreTextureSampleBase(const FMediaIOCoreTextureSampleBase&) = default;
+	FMediaIOCoreTextureSampleBase(FMediaIOCoreTextureSampleBase&&) = default;
+	FMediaIOCoreTextureSampleBase& operator=(const FMediaIOCoreTextureSampleBase&) = default;
+	FMediaIOCoreTextureSampleBase& operator=(FMediaIOCoreTextureSampleBase&&) = default;
+	PRAGMA_ENABLE_DEPRECATION_WARNINGS
+
 	/**
 	 * Initialize the sample.
 	 *
