Fix CIS in MediaIOCoreUtilities.h caused by multiply defined symbols

diff --git a/Engine/Plugins/Media/MediaIOFramework/Source/MediaIOCore/Public/MediaIOCoreUtilities.h b/Engine/Plugins/Media/MediaIOFramework/Source/MediaIOCore/Public/MediaIOCoreUtilities.h
--- a/Engine/Plugins/Media/MediaIOFramework/Source/MediaIOCore/Public/MediaIOCoreUtilities.h
+++ b/Engine/Plugins/Media/MediaIOFramework/Source/MediaIOCore/Public/MediaIOCoreUtilities.h
@@ -16,9 +16,9 @@
 namespace UE::MediaIO
 {
 #if NO_LOGGING
-	void LogThrottle(const FNoLoggingCategory& InLogCategory, ELogVerbosity::Type InVerbosity, const FTimespan& InTimeBetweenLogs, const FString& LogDetails, const FString& FileName, int32 LineNumber);
+	static void LogThrottle(const FNoLoggingCategory& InLogCategory, ELogVerbosity::Type InVerbosity, const FTimespan& InTimeBetweenLogs, const FString& LogDetails, const FString& FileName, int32 LineNumber);
 #else
-	void LogThrottle(const FLogCategoryBase& InLogCategory, ELogVerbosity::Type InVerbosity, const FTimespan& InTimeBetweenLogs, const FString& LogDetails, const FString& FileName, int32 LineNumber);
+	static void LogThrottle(const FLogCategoryBase& InLogCategory, ELogVerbosity::Type InVerbosity, const FTimespan& InTimeBetweenLogs, const FString& LogDetails, const FString& FileName, int32 LineNumber);
 #endif
 }
 
@@ -27,11 +27,11 @@ namespace UE::MediaIO
 namespace UE::MediaIO
 {
 #if NO_LOGGING
-	void LogThrottle(const FNoLoggingCategory& InLogCategory, ELogVerbosity::Type InVerbosity, const FTimespan& InTimeBetweenLogs, const FString& LogDetails, const FString& FileName, int32 LineNumber)
+	static void LogThrottle(const FNoLoggingCategory& InLogCategory, ELogVerbosity::Type InVerbosity, const FTimespan& InTimeBetweenLogs, const FString& LogDetails, const FString& FileName, int32 LineNumber)
 	{
 	}
 #else
-	void LogThrottle(const FLogCategoryBase& InLogCategory, ELogVerbosity::Type InVerbosity, const FTimespan& InTimeBetweenLogs, const FString& LogDetails, const FString& FileName, int32 LineNumber)
+	static void LogThrottle(const FLogCategoryBase& InLogCategory, ELogVerbosity::Type InVerbosity, const FTimespan& InTimeBetweenLogs, const FString& LogDetails, const FString& FileName, int32 LineNumber)
 	{
 		// This code was adapted from AudioMixer.h to support throttling instead of only logging once.
 		// Log once to avoid Spam.
