Re-enable media capture resample pass

diff --git a/Engine/Plugins/Media/MediaIOFramework/MediaIOFramework.uplugin b/Engine/Plugins/Media/MediaIOFramework/MediaIOFramework.uplugin
--- a/Engine/Plugins/Media/MediaIOFramework/MediaIOFramework.uplugin
+++ b/Engine/Plugins/Media/MediaIOFramework/MediaIOFramework.uplugin
@@ -17,7 +17,7 @@
 		{
 			"Name": "MediaIOCore",
 			"Type": "Runtime",
-			"LoadingPhase": "Default"
+			"LoadingPhase": "PostConfigInit"
 		},
 		{
 			"Name": "MediaIOEditor",
diff --git a/Engine/Plugins/Media/MediaIOFramework/Shaders/MediaIO.usf b/Engine/Plugins/Media/MediaIOFramework/Shaders/MediaIO.usf
new file mode 100644
--- /dev/null
+++ b/Engine/Plugins/Media/MediaIOFramework/Shaders/MediaIO.usf
@@ -0,0 +1,40 @@
+// Copyright Epic Games, Inc. All Rights Reserved.
+
+#include "/Engine/Private/ScreenPass.ush"
+#include "/Engine/Private/Common.ush"
+
+SCREEN_PASS_TEXTURE_VIEWPORT(Input)
+SCREEN_PASS_TEXTURE_VIEWPORT(Output)
+
+Texture2D InputTexture;
+
+SamplerState InputSampler;
+
+// Shader to downsample a texture, Adapted from PostProcessDownsample.usf
+void MediaDownsample(float4 SvPosition : SV_POSITION, out float4 OutColor : SV_Target0)
+{
+	const float2 UV = SvPosition.xy * Output_ExtentInverse;
+
+	// Output: float4(RGBA), 4 filtered samples
+	float2 UVs[4];
+
+	// Blur during downsample (4x4 kernel) to get better quality especially for HDR content.
+	UVs[0] = UV + Input_ExtentInverse * float2(-1, -1);
+	UVs[1] = UV + Input_ExtentInverse * float2( 1, -1);
+	UVs[2] = UV + Input_ExtentInverse * float2(-1,  1);
+	UVs[3] = UV + Input_ExtentInverse * float2( 1,  1);
+
+	float4 Sample[4];
+
+	UNROLL
+	for(uint i = 0; i < 4; ++i)
+	{
+		float2 UVAtIndex = clamp(UVs[i], Input_UVViewportBilinearMin, Input_UVViewportBilinearMax);
+		Sample[i] = Texture2DSampleLevel(InputTexture, InputSampler, UVAtIndex, 0);
+	}
+
+	OutColor = (Sample[0] + Sample[1] + Sample[2] + Sample[3]) * 0.25f;
+
+	// Fixed rarely occurring yellow color tint of the whole viewport (certain viewport size, need to investigate more)
+	OutColor.rgb = max(float3(0,0,0), OutColor.rgb);
+}
diff --git a/Engine/Plugins/Media/MediaIOFramework/Source/MediaIOCore/Private/MediaCaptureRenderPass.cpp b/Engine/Plugins/Media/MediaIOFramework/Source/MediaIOCore/Private/MediaCaptureRenderPass.cpp
--- a/Engine/Plugins/Media/MediaIOFramework/Source/MediaIOCore/Private/MediaCaptureRenderPass.cpp
+++ b/Engine/Plugins/Media/MediaIOFramework/Source/MediaIOCore/Private/MediaCaptureRenderPass.cpp
@@ -24,77 +24,20 @@ DECLARE_GPU_STAT(MediaCapture_Conversion);
 
 namespace UE::MediaCapture::Resample
 {
-	/** Shader parameter struct used for the Resample pixel shader. */
-	BEGIN_SHADER_PARAMETER_STRUCT(FMediaCaptureResampleParameters, )
-		SHADER_PARAMETER_STRUCT_REF(FViewUniformShaderParameters, ViewUniformBuffer)
-		SHADER_PARAMETER_STRUCT(FScreenPassTextureViewportParameters, Input)
-		SHADER_PARAMETER_STRUCT(FScreenPassTextureViewportParameters, Output)
-		SHADER_PARAMETER_RDG_TEXTURE(Texture2D, InputTexture)
-		SHADER_PARAMETER_SAMPLER(SamplerState, InputSampler)
-	END_SHADER_PARAMETER_STRUCT()
-
 	/**
-	 * This was adapted from PostProcessDownsample to support any output texture size.
+	 * Pixel shader to resample a texture using 4 samples per pixel.
 	 */
-	enum class EResampleQuality : uint8
-	{
-		// Single filtered sample (2x2 tap).
-		Low,
-
-		// Four filtered samples (4x4 tap).
-		High,
-
-		MAX
-	};
-
-	// The set of inputs needed to add a Resample pass to RDG.
-	struct FResamplePassInputs
-	{
-		FResamplePassInputs() = default;
-
-		// Friendly name of the pass. Used for logging and profiling.
-		const TCHAR* Name = nullptr;
-
-		// Input scene color RDG texture / view rect. Must not be null.
-		FScreenPassTexture SceneColor;
-
-		// The Resample method to use.
-		EResampleQuality Quality = EResampleQuality::Low;
-	};
-
-	/** Instantiate a Resample parameters struct. */
-	FMediaCaptureResampleParameters GetResampleParameters(const FViewInfo& View, FScreenPassTexture Output, FScreenPassTexture Input, EResampleQuality ResampleMethod)
-	{
-		check(Output.IsValid());
-		check(Input.IsValid());
-
-		const FScreenPassTextureViewportParameters InputParameters = GetScreenPassTextureViewportParameters(FScreenPassTextureViewport(Input));
-		const FScreenPassTextureViewportParameters OutputParameters = GetScreenPassTextureViewportParameters(FScreenPassTextureViewport(Output));
-
-		FMediaCaptureResampleParameters Parameters;
-		Parameters.ViewUniformBuffer = View.ViewUniformBuffer;
-		Parameters.Input = InputParameters;
-		Parameters.Output = OutputParameters;
-		Parameters.InputTexture = Input.Texture;
-		Parameters.InputSampler = TStaticSamplerState<SF_Bilinear, AM_Clamp, AM_Clamp, AM_Clamp>::GetRHI();
-		return Parameters;
-	}
-
-	class FResampleQualityDimension : SHADER_PERMUTATION_ENUM_CLASS("DOWNSAMPLE_QUALITY", EResampleQuality);
-	using FResamplePermutationDomain = TShaderPermutationDomain<FResampleQualityDimension>;
-
-	/** Resample pixel shader struct copied from PostProcessResample.h */
-	/*
 	class FMediaCaptureResamplePS : public FGlobalShader
 	{
 	public:
 		DECLARE_GLOBAL_SHADER(FMediaCaptureResamplePS);
 		SHADER_USE_PARAMETER_STRUCT(FMediaCaptureResamplePS, FGlobalShader);
 
-		using FPermutationDomain = FResamplePermutationDomain;
-
 		BEGIN_SHADER_PARAMETER_STRUCT(FParameters, )
-			SHADER_PARAMETER_STRUCT_INCLUDE(FMediaCaptureResampleParameters, Common)
+			SHADER_PARAMETER_STRUCT(FScreenPassTextureViewportParameters, Input)
+			SHADER_PARAMETER_STRUCT(FScreenPassTextureViewportParameters, Output)
+			SHADER_PARAMETER_RDG_TEXTURE(Texture2D, InputTexture) 
+			SHADER_PARAMETER_SAMPLER(SamplerState, InputSampler)
 			RENDER_TARGET_BINDING_SLOTS()
 		END_SHADER_PARAMETER_STRUCT()
 
@@ -102,15 +45,30 @@ namespace UE::MediaCapture::Resample
 		{
 			return true;
 		}
-	};
-	*/
 
-	//IMPLEMENT_GLOBAL_SHADER(FMediaCaptureResamplePS, "/Engine/Private/PostProcessDownsample.usf", "MainPS", SF_Pixel);
+		/** Allocates and setup shader parameter in the incoming graph builder */
+		static FMediaCaptureResamplePS::FParameters* AllocateAndSetParameters(FRDGBuilder& GraphBuilder, FRDGTextureRef RGBATexture, FRDGTextureRef OutputTexture)
+		{
+			FMediaCaptureResamplePS::FParameters* Parameters = GraphBuilder.AllocParameters<FMediaCaptureResamplePS::FParameters>();
+
+			Parameters->Input = GetScreenPassTextureViewportParameters(FScreenPassTextureViewport(RGBATexture));
+			Parameters->Output = GetScreenPassTextureViewportParameters(FScreenPassTextureViewport(OutputTexture));
+			Parameters->InputTexture = RGBATexture;
+			Parameters->InputSampler = TStaticSamplerState<SF_Bilinear, AM_Clamp, AM_Clamp, AM_Clamp>::GetRHI();
+			Parameters->RenderTargets[0] = FRenderTargetBinding{ OutputTexture, ERenderTargetLoadAction::ENoAction };
 
+			return Parameters;
+		}
+	};
+	
+	IMPLEMENT_GLOBAL_SHADER(FMediaCaptureResamplePS, "/MediaIOShaders/MediaIO.usf", "MediaDownsample", SF_Pixel);
+	
 	/** Adds a Resample pass to the graph builder. */
-	void AddResamplePass(FRDGBuilder& GraphBuilder, const FViewInfo& View, const FResamplePassInputs& Inputs, FRDGTextureRef OutputTexture)
+	void AddResamplePass(FRDGBuilder& GraphBuilder, const FViewInfo& View, const FScreenPassTexture& InputTexture, FRDGTextureRef OutputTexture)
 	{
-		check(Inputs.SceneColor.IsValid());
+		check(InputTexture.Texture);
+
+		const FScreenPassRenderTarget Input = FScreenPassRenderTarget::CreateFromInput(GraphBuilder, InputTexture, ERenderTargetLoadAction::ENoAction, TEXT("MediaCaptureResampleInput"));
 
 		FScreenPassRenderTarget Output;
 		const FIntVector OutputSize = OutputTexture->Desc.GetSize();
@@ -118,29 +76,20 @@ namespace UE::MediaCapture::Resample
 		Output.ViewRect = FIntRect{ 0, 0, OutputSize.X, OutputSize.Y };
 		Output.LoadAction = ERenderTargetLoadAction::ENoAction;
 
-		FResamplePermutationDomain PermutationVector;
-		PermutationVector.Set<FResampleQualityDimension>(Inputs.Quality);
-
-		/*
-
-		FMediaCaptureResamplePS::FParameters* PassParameters = GraphBuilder.AllocParameters<FMediaCaptureResamplePS::FParameters>();
-		PassParameters->Common = GetResampleParameters(View, Output, Inputs.SceneColor, Inputs.Quality);
-		PassParameters->RenderTargets[0] = Output.GetRenderTargetBinding();
-
-		TShaderMapRef<FMediaCaptureResamplePS> PixelShader(View.ShaderMap, PermutationVector);
+		FMediaCaptureResamplePS::FParameters* PassParameters = FMediaCaptureResamplePS::AllocateAndSetParameters(GraphBuilder, InputTexture.Texture, OutputTexture);
+		
+		const FGlobalShaderMap* GlobalShaderMap = GetGlobalShaderMap(GMaxRHIFeatureLevel);
+		const TShaderMapRef<FMediaCaptureResamplePS> PixelShader(GlobalShaderMap);
 		FPixelShaderUtils::AddFullscreenPass(
 			GraphBuilder,
-			View.ShaderMap,
-			RDG_EVENT_NAME("MediaCapture Resample(%s Quality=%s) %dx%d -> %dx%d",
+			GlobalShaderMap,
+			RDG_EVENT_NAME("MediaCapture Resample(%s) %dx%d -> %dx%d",
 				Output.Texture->Name,
-				Inputs.Quality == EResampleQuality::High ? TEXT("High") : TEXT("Bilinear"),
-				Inputs.SceneColor.ViewRect.Width(), Inputs.SceneColor.ViewRect.Height(),
+				Input.ViewRect.Width(), Input.ViewRect.Height(),
 				Output.ViewRect.Width(), Output.ViewRect.Height()),
 			PixelShader,
 			PassParameters,
 			Output.ViewRect);
-
-		*/
 	}
 
 }
@@ -421,8 +370,7 @@ namespace UE::MediaCapture
 
 		if (MediaCapture->GetDesiredCaptureOptions().ResizeMethod == EMediaCaptureResizeMethod::ResizeInRenderPass)
 		{
-			// Disabled temporarily to fix CIS issue related to the shader.
-			//RenderPasses.Add(CreateResamplePass());
+			RenderPasses.Add(CreateResamplePass());
 		}
 		
 		if (MediaCapture->GetDesiredCaptureOptions().ColorConversionSettings.IsValid())
@@ -445,7 +393,7 @@ namespace UE::MediaCapture
 
 		check(OutputResource->Type == FRDGTexture::StaticType);
 		FRDGTexture* OutputTexture =  static_cast<FRDGTexture*>(OutputResource);
-	
+		
 		//Based on conversion type, this might be changed
 		bool bRequiresFormatConversion = ConversionPassArgs.bRequiresFormatConversion;
 
@@ -522,11 +470,6 @@ namespace UE::MediaCapture
 		// Downsampling only supports texture at the moment.
 		check(OutputTexture->Type == ERDGViewableResourceType::Texture);
 			
-		UE::MediaCapture::Resample::FResamplePassInputs PassInputs;
-		PassInputs.Name = TEXT("MediaCapture.Resample");
-		PassInputs.SceneColor = FScreenPassTexture(ConversionPassArgs.SourceRGBTexture);
-		PassInputs.Quality = MediaCapture::Resample::EResampleQuality::High;
-
 		const FIntRect ViewRect(ConversionPassArgs.CopyInfo.GetSourceRect());
 
 		//Dummy ViewFamily/ViewInfo created to use built in Draw Screen/Texture Pass
@@ -542,7 +485,7 @@ namespace UE::MediaCapture
 		FViewInfo ViewInfo = FViewInfo(ViewInitOptions);
 
 		RDG_GPU_STAT_SCOPE(Args.GraphBuilder, MediaCapture_Resample);
-		MediaCapture::Resample::AddResamplePass(Args.GraphBuilder, ViewInfo, PassInputs, (FRDGTextureRef)OutputTexture);
+		MediaCapture::Resample::AddResamplePass(Args.GraphBuilder, ViewInfo, FScreenPassTexture(ConversionPassArgs.SourceRGBTexture), (FRDGTextureRef)OutputTexture);
 	}
 
 	FRenderPipeline::FRenderPipeline(UMediaCapture* InMediaCapture)
diff --git a/Engine/Plugins/Media/MediaIOFramework/Source/MediaIOCore/Private/MediaIOCoreModule.cpp b/Engine/Plugins/Media/MediaIOFramework/Source/MediaIOCore/Private/MediaIOCoreModule.cpp
--- a/Engine/Plugins/Media/MediaIOFramework/Source/MediaIOCore/Private/MediaIOCoreModule.cpp
+++ b/Engine/Plugins/Media/MediaIOFramework/Source/MediaIOCore/Private/MediaIOCoreModule.cpp
@@ -1,10 +1,13 @@
 // Copyright Epic Games, Inc. All Rights Reserved.
 
 #include "MediaIOCoreModule.h"
+
 #include "IMediaIOCoreDeviceProvider.h"
 #include "IMediaIOCoreModule.h"
-
+#include "Interfaces/IPluginManager.h"
+#include "Misc/Paths.h"
 #include "Modules/ModuleManager.h"
+#include "ShaderCore.h"
 
 
 DEFINE_LOG_CATEGORY(LogMediaIOCore);
@@ -15,6 +18,12 @@ DEFINE_LOG_CATEGORY(LogMediaIOCore);
 class FMediaIOCoreModule : public IMediaIOCoreModule
 {
 public:
+	virtual void StartupModule() override
+	{
+		FString PluginShaderDir = FPaths::Combine(IPluginManager::Get().FindPlugin(TEXT("MediaIOFramework"))->GetBaseDir(), TEXT("Shaders"));
+		AddShaderSourceDirectoryMapping(TEXT("/MediaIOShaders"), PluginShaderDir);
+	}
+	
 	virtual void RegisterDeviceProvider(IMediaIOCoreDeviceProvider* InProvider) override
 	{
 		if (InProvider)
