- Adding media capture markup to track frames
- Moving sync point to its own thread for frame locking work

diff --git a/Engine/Plugins/Media/MediaIOFramework/Source/MediaIOCore/Private/MediaCapture.cpp b/Engine/Plugins/Media/MediaIOFramework/Source/MediaIOCore/Private/MediaCapture.cpp
--- a/Engine/Plugins/Media/MediaIOFramework/Source/MediaIOCore/Private/MediaCapture.cpp
+++ b/Engine/Plugins/Media/MediaIOFramework/Source/MediaIOCore/Private/MediaCapture.cpp
@@ -13,6 +13,7 @@
 #include "MediaCaptureRenderPass.h"
 #include "MediaCaptureSceneViewExtension.h"
 #include "MediaCaptureSources.h"
+#include "MediaCaptureSyncPointWatcher.h"
 #include "MediaIOFrameManager.h"
 #include "MediaIOCoreModule.h"
 #include "MediaOutput.h"
@@ -277,6 +278,7 @@ bool UMediaCapture::StartSourceCapture(TSharedPtr<UE::MediaCapture::Private::FCa
 	bShouldCaptureRHIResource = ShouldCaptureRHIResource();
 
 	CaptureRenderPipeline = MakePimpl<UE::MediaCapture::FRenderPipeline>(this);
+	SyncPointWatcher = MakePimpl<UE::MediaCaptureData::FSyncPointWatcher>(this);
 
 	if (bInitialized)
 	{
@@ -545,13 +547,14 @@ void UMediaCapture::StopCapture(bool bAllowPendingFrameToBeProcess)
 				CaptureSource->ResetSourceBufferSize(bFlushRenderingCommands);
 				CaptureSource.Reset();
 			}
-			
+
 			DesiredSize = FIntPoint(1280, 720);
 			DesiredPixelFormat = EPixelFormat::PF_A2B10G10R10;
 			DesiredOutputSize = FIntPoint(1280, 720);
 			DesiredOutputPixelFormat = EPixelFormat::PF_A2B10G10R10;
 			DesiredCaptureOptions = FMediaCaptureOptions();
 			ConversionOperation = EMediaCaptureConversionOperation::NONE;
+			SyncPointWatcher.Reset();
 			ViewExtension.Reset();
 			CaptureRenderPipeline.Reset();
 			
@@ -618,6 +621,7 @@ void UMediaCapture::CaptureImmediate_RenderThread(const UE::MediaCaptureData::FC
 {
 	using namespace UE::MediaCaptureData;
 	TRACE_CPUPROFILER_EVENT_SCOPE(UMediaCapture::CaptureImmediate_RenderThread);
+	TRACE_CPUPROFILER_EVENT_SCOPE_TEXT(*FString::Printf(TEXT("MediaCapturePipe: %llu"), GFrameCounterRenderThread));
 	
 	check(IsInRenderingThread());
 	
@@ -908,24 +912,7 @@ void UMediaCapture::WaitForSingleExperimentalSchedulingTaskToComplete()
 	if (UseExperimentalScheduling())
 	{
 		// Presumably the rendering thread could be in the process of dispatching the task
-		FScopeLock ScopedLock(&PendingReadbackTasksCriticalSection);
-		if (PendingReadbackTasks.Num())
-		{
-			auto It = PendingReadbackTasks.CreateIterator();
-
-			// Go through the list to find a pending task.
-			while (It && It->IsCompleted())
-			{
-				It.RemoveCurrent();
-				++It;
-			}
-
-			if (It)
-			{
-				// We have a pending task, wait until it's completeted.
-				It->BusyWait();
-			}
-		}
+		SyncPointWatcher->WaitForSinglePendingTaskToComplete();
 
 		if (!UseAnyThreadCapture())
 		{
@@ -938,40 +925,12 @@ void UMediaCapture::WaitForSingleExperimentalSchedulingTaskToComplete()
 	}
 }
 
-void UMediaCapture::CleanupCompletedExperimentalSchedulingTasks()
-{
-	FScopeLock ScopedLock(&PendingReadbackTasksCriticalSection);
-
-	for (auto It = PendingReadbackTasks.CreateIterator(); It; ++It)
-	{
-		if (It->IsCompleted())
-		{
-			It.RemoveCurrent();
-		}
-	}
-}
-
 void UMediaCapture::WaitForAllExperimentalSchedulingTasksToComplete()
 {
 	if (UseExperimentalScheduling())
 	{
-		FScopeLock ScopedLock(&PendingReadbackTasksCriticalSection);
-
-		// Clean up completed tasks.
-		for (auto It = PendingReadbackTasks.CreateIterator(); It; ++It)
-		{
-			if (It->IsCompleted())
-			{
-				It.RemoveCurrent();
-			}
-			else
-			{
-				It->BusyWait();
-				It.RemoveCurrent();
-			}
-		}
+		SyncPointWatcher->WaitForAllPendingTasksToComplete();
 
-		
 		if (!UseAnyThreadCapture())
 		{
 			// This code might have dispatched a task on the render thread, so we need to wait again
@@ -997,8 +956,6 @@ void UMediaCapture::ProcessCapture_GameThread()
 {
 	using namespace UE::MediaCaptureData;
 	
-	CleanupCompletedExperimentalSchedulingTasks();
-
 	// Acquire a frame
 	TSharedPtr<FCaptureFrame> CapturingFrame;
 	if(GetState() != EMediaCaptureState::StopRequested)
@@ -1338,19 +1295,6 @@ void UMediaCapture::OnEndFrame_GameThread()
 		return;
 	}
 
-	{
-		FScopeLock ScopedLock(&PendingReadbackTasksCriticalSection);
-
-		// Clean up completed tasks.
-		for (auto It = PendingReadbackTasks.CreateIterator(); It; ++It)
-		{
-			if (It->IsCompleted())
-			{
-				It.RemoveCurrent();
-			}
-		}
-	}
-
 	ProcessCapture_GameThread();
 }
 
diff --git a/Engine/Plugins/Media/MediaIOFramework/Source/MediaIOCore/Private/MediaCaptureHelper.cpp b/Engine/Plugins/Media/MediaIOFramework/Source/MediaIOCore/Private/MediaCaptureHelper.cpp
--- a/Engine/Plugins/Media/MediaIOFramework/Source/MediaIOCore/Private/MediaCaptureHelper.cpp
+++ b/Engine/Plugins/Media/MediaIOFramework/Source/MediaIOCore/Private/MediaCaptureHelper.cpp
@@ -2,6 +2,7 @@
 
 #include "MediaCaptureHelper.h"
 #include "MediaCaptureRenderPass.h"
+#include "MediaCaptureSyncPointWatcher.h"
 
 DECLARE_CYCLE_STAT(TEXT("MediaCapture RenderThread LockResource"), STAT_MediaCaptureHelper_RenderThread_LockResource, STATGROUP_Media);
 DECLARE_CYCLE_STAT(TEXT("MediaCapture AnyThread LockResource"), STAT_MediaCaptureHelper_AnyThread_LockResource, STATGROUP_Media);
@@ -11,6 +12,7 @@ DECLARE_CYCLE_STAT(TEXT("MediaCapture RenderThread FrameCapture"), STAT_MediaCap
 DECLARE_CYCLE_STAT(TEXT("MediaCapture RenderThread RHI Capture Callback"), STAT_MediaCaptureHelper_RenderThread_RHI_CaptureCallback, STATGROUP_Media);
 
 DECLARE_GPU_STAT(MediaCapture_CaptureFrame);
+DECLARE_GPU_STAT(MediaCapture_SyncPoint);
 DECLARE_GPU_STAT(MediaCapture_CustomCapture);
 DECLARE_GPU_STAT(MediaCapture_Readback);
 
@@ -126,14 +128,11 @@ bool FMediaCaptureHelper::CaptureFrame(const UE::MediaCaptureData::FCaptureFrame
 		return false;
 	}
 
+	if (CapturingFrame->IsTextureResource())
 	{
 		TRACE_CPUPROFILER_EVENT_SCOPE(UMediaCapture::LockDMATexture_RenderThread);
 		TRACE_CPUPROFILER_EVENT_SCOPE_TEXT(*FString::Printf(TEXT("LockDmaTexture Output Frame %d"), CapturingFrame->CaptureBaseData.SourceFrameNumberRenderThread));
-
-		if (CapturingFrame->IsTextureResource())
-		{
-			Args.MediaCapture->LockDMATexture_RenderThread(CapturingFrame->GetTextureResource());
-		}
+		Args.MediaCapture->LockDMATexture_RenderThread(CapturingFrame->GetTextureResource());
 	}
 
 	{
@@ -169,6 +168,7 @@ bool FMediaCaptureHelper::CaptureFrame(const UE::MediaCaptureData::FCaptureFrame
 
 		if (Args.MediaCapture->UseExperimentalScheduling())
 		{
+			RDG_GPU_STAT_SCOPE(Args.GraphBuilder, MediaCapture_SyncPoint);
 			if (CapturingFrame->IsTextureResource())
 			{
 				AddSyncPointPass<FTextureCaptureFrame>(Args.GraphBuilder, Args.MediaCapture, CapturingFrame, FinalPassOutputResource);
@@ -195,86 +195,12 @@ void FMediaCaptureHelper::ExecuteSyncPointPass(FRHICommandListImmediate& RHICmdL
 			RHICmdList.WriteGPUFence(SyncDataPtr->RHIFence);
 			SyncDataPtr->bIsBusy = true;
 
-			// Here we request a number to process the frames in order.
-			const uint32 ExecutionNumber = MediaCapture->OrderedAsyncGateCaptureReady.GetANumber();
-
-			// Spawn a task that will wait (poll) and continue the process of providing a new texture
-			UE::Tasks::FTask Task = UE::Tasks::Launch(UE_SOURCE_LOCATION, [MediaCapture, CapturingFrame, SyncDataPtr, ExecutionNumber]()
-			{
-				TRACE_CPUPROFILER_EVENT_SCOPE(UMediaCapture::SyncPass);
-
-				ON_SCOPE_EXIT
-				{
-					// Make sure that we give up our turn to execute when exiting this task,
-					// which will allow the other async tasks to execute after waiting for the fence.
-					MediaCapture->OrderedAsyncGateCaptureReady.GiveUpTurn(ExecutionNumber);
-				};
-
-				double WaitTime = 0.0;
-				bool bWaitedForCompletion = false;
-				{
-					FScopedDurationTimer Timer(WaitTime);
-
-					// Wait until fence has been written (shader has completed)
-					while (true)
-					{
-						if (SyncDataPtr->RHIFence->Poll())
-						{
-							bWaitedForCompletion = true;
-							break;
-						}
-
-						if (!CapturingFrame->bMediaCaptureActive)
-						{
-							bWaitedForCompletion = false;
-							break;
-						}
-
-						constexpr float SleepTimeSeconds = 50 * 1E-6;
-						FPlatformProcess::SleepNoStats(SleepTimeSeconds);
-					}
-
-					SyncDataPtr->RHIFence->Clear();
-					SyncDataPtr->bIsBusy = false;
-				}
-
-				if (CapturingFrame->bMediaCaptureActive && bWaitedForCompletion && MediaCapture)
-				{
-					// Ensure that we do not run the following code out of order with respect to the other sibling async tasks,
-					// because the Pending Frames are expected to be processed in order.
-					if (!MediaCapture->OrderedAsyncGateCaptureReady.IsMyTurn(ExecutionNumber))
-					{
-						TRACE_CPUPROFILER_EVENT_SCOPE(UMediaCapture::SyncPass::OutOfOrderWait);
-
-						UE_LOG(LogMediaIOCore, Warning, TEXT(
-							"The wait for the GPU fence for the next frame of MediaCapture '%s' unexpectedly happened out of order. "
-							"Order will be enforced by waiting until the previous frames are processed. "),
-							*MediaCapture->MediaOutputName);
-
-						MediaCapture->OrderedAsyncGateCaptureReady.WaitForTurn(ExecutionNumber);
-					}
-
-					if (MediaCapture->UseAnyThreadCapture())
-					{
-						OnReadbackComplete(FRHICommandListExecutor::GetImmediateCommandList(), MediaCapture, CapturingFrame);
-					}
-					else
-					{
-						++MediaCapture->WaitingForRenderCommandExecutionCounter;
-
-						ENQUEUE_RENDER_COMMAND(MediaOutputCaptureReadbackComplete)([MediaCapture, CapturingFrame](FRHICommandList& RHICommandList)
-							{
-								OnReadbackComplete(RHICommandList, MediaCapture, CapturingFrame);
-							});
-					}
-				}
-
-			}, static_cast<LowLevelTasks::ETaskPriority>(CVarMediaIOCapturePollTaskPriority.GetValueOnRenderThread()));
-
+			// Queue capture data and our thread will wait (poll) and continue the process of providing a new texture
 			{
-				// Add the task to the capture's list of pending tasks.
-				FScopeLock ScopedLock(&MediaCapture->PendingReadbackTasksCriticalSection);
-				MediaCapture->PendingReadbackTasks.Add(MoveTemp(Task));
+				UE::MediaCaptureData::FPendingCaptureData NewCaptureData;
+				NewCaptureData.SyncHandler = SyncDataPtr;
+				NewCaptureData.CapturingFrame = CapturingFrame;
+				MediaCapture->SyncPointWatcher->QueuePendingCapture(MoveTemp(NewCaptureData));
 			}
 		}
 		else
diff --git a/Engine/Plugins/Media/MediaIOFramework/Source/MediaIOCore/Private/MediaCaptureHelper.h b/Engine/Plugins/Media/MediaIOFramework/Source/MediaIOCore/Private/MediaCaptureHelper.h
--- a/Engine/Plugins/Media/MediaIOFramework/Source/MediaIOCore/Private/MediaCaptureHelper.h
+++ b/Engine/Plugins/Media/MediaIOFramework/Source/MediaIOCore/Private/MediaCaptureHelper.h
@@ -19,6 +19,8 @@ namespace UE::MediaCaptureData
 	class FMediaCaptureHelper
 	{
 	public:
+		friend class FSyncPointWatcher;
+		
 		/** Validates the input and output size and format compatibility. */
 		static bool AreInputsValid(const UE::MediaCaptureData::FCaptureFrameArgs& Args);
 
diff --git a/Engine/Plugins/Media/MediaIOFramework/Source/MediaIOCore/Private/MediaCaptureSyncPointWatcher.cpp b/Engine/Plugins/Media/MediaIOFramework/Source/MediaIOCore/Private/MediaCaptureSyncPointWatcher.cpp
new file mode 100644
--- /dev/null
+++ b/Engine/Plugins/Media/MediaIOFramework/Source/MediaIOCore/Private/MediaCaptureSyncPointWatcher.cpp
@@ -0,0 +1,145 @@
+// Copyright Epic Games, Inc. All Rights Reserved.
+
+#include "MediaCaptureSyncPointWatcher.h"
+
+#include "HAL/RunnableThread.h"
+#include "MediaCaptureHelper.h"
+
+
+namespace UE::MediaCaptureData
+{
+
+FSyncPointWatcher::FSyncPointWatcher(UMediaCapture* InMediaCapture)
+	: MediaCapture(InMediaCapture)
+{
+	bIsEnabled = true;
+
+	const FString ThreadName = FString::Printf(TEXT("%s_SyncThread"), *MediaCapture->GetName());
+	WorkingThread.Reset(FRunnableThread::Create(this, *ThreadName, 32 * 1024, TPri_Highest));
+}
+
+uint32 FSyncPointWatcher::Run()
+{
+	while (bIsEnabled)
+	{
+		if (PendingCaptureFrames.IsEmpty())
+		{
+			PendingTaskSignal->Wait();
+			if (!bIsEnabled)
+			{
+				break;
+			}
+		}
+
+		FPendingCaptureData NextCaptureData;
+		{
+			if(!PendingCaptureFrames.Dequeue(NextCaptureData))
+			{
+				continue;
+			}
+		}
+
+		ProcessPendingCapture(NextCaptureData);
+	}
+
+	return 0;
+}
+
+void FSyncPointWatcher::ProcessPendingCapture(const FPendingCaptureData& Data)
+{
+	ON_SCOPE_EXIT
+	{
+		--PendingCaptureFramesCount;
+	};
+
+	double WaitTime = 0.0;
+	bool bWaitedForCompletion = false;
+	{
+		FScopedDurationTimer Timer(WaitTime);
+
+		// Wait until fence has been written (shader has completed)
+		while (true)
+		{
+			if (Data.SyncHandler->RHIFence->Poll())
+			{
+				bWaitedForCompletion = true;
+				break;
+			}
+
+			if (!Data.CapturingFrame->bMediaCaptureActive)
+			{
+				bWaitedForCompletion = false;
+				break;
+			}
+
+			constexpr float SleepTimeSeconds = 50 * 1E-6;
+			FPlatformProcess::SleepNoStats(SleepTimeSeconds);
+		}
+
+		Data.SyncHandler->RHIFence->Clear();
+		Data.SyncHandler->bIsBusy = false;
+	}
+
+	if (Data.CapturingFrame->bMediaCaptureActive && bWaitedForCompletion && MediaCapture)
+	{
+		// Ensure that we do not run the following code out of order with respect to the other sibling async tasks,
+		// because the Pending Frames are expected to be processed in order.
+
+		if (MediaCapture->UseAnyThreadCapture())
+		{
+			FMediaCaptureHelper::OnReadbackComplete(FRHICommandListExecutor::GetImmediateCommandList(), MediaCapture, Data.CapturingFrame);
+		}
+		else
+		{
+			++MediaCapture->WaitingForRenderCommandExecutionCounter;
+
+			ENQUEUE_RENDER_COMMAND(MediaOutputCaptureReadbackComplete)([OwnerCapture = MediaCapture, CapturingFrame = Data.CapturingFrame](FRHICommandList& RHICommandList)
+			{
+				FMediaCaptureHelper::OnReadbackComplete(RHICommandList, OwnerCapture, CapturingFrame);
+			});
+		}
+	}
+}
+
+void FSyncPointWatcher::QueuePendingCapture(FPendingCaptureData NewData)
+{
+	++PendingCaptureFramesCount;
+	PendingCaptureFrames.Enqueue(MoveTemp(NewData));
+	PendingTaskSignal->Trigger();
+}
+
+void FSyncPointWatcher::Stop()
+{
+	bIsEnabled = false;
+	PendingTaskSignal->Trigger();
+}
+
+void FSyncPointWatcher::WaitForAllPendingTasksToComplete()
+{
+	// Clean up completed tasks.
+	while (PendingCaptureFramesCount > 0)
+	{
+		FPlatformProcess::SleepNoStats(50e-6);
+	}
+}
+
+void FSyncPointWatcher::WaitForSinglePendingTaskToComplete()
+{
+	const uint32 PendingSnapshot = PendingCaptureFramesCount;
+	if (PendingSnapshot > 0)
+	{
+		// Wait for at least one frame to be completed
+		while (PendingCaptureFramesCount.load() >= PendingSnapshot)
+		{
+			FPlatformProcess::SleepNoStats(50e-6);
+		}
+	}
+}
+
+FSyncPointWatcher::~FSyncPointWatcher()
+{
+	Stop();
+	WorkingThread->WaitForCompletion();
+}
+
+}
\ No newline at end of file
diff --git a/Engine/Plugins/Media/MediaIOFramework/Source/MediaIOCore/Private/MediaCaptureSyncPointWatcher.h b/Engine/Plugins/Media/MediaIOFramework/Source/MediaIOCore/Private/MediaCaptureSyncPointWatcher.h
new file mode 100644
--- /dev/null
+++ b/Engine/Plugins/Media/MediaIOFramework/Source/MediaIOCore/Private/MediaCaptureSyncPointWatcher.h
@@ -0,0 +1,73 @@
+// Copyright Epic Games, Inc. All Rights Reserved.
+
+#pragma once
+
+#include "HAL/Runnable.h"
+
+#include "HAL/ThreadManager.h"
+#include "MediaCapture.h"
+
+class FRunnableThread;
+
+namespace UE::MediaCaptureData
+{
+	struct FPendingCaptureData
+	{
+		/** Sync handler holding the fence used for this frame */
+		TSharedPtr<UMediaCapture::FMediaCaptureSyncData> SyncHandler;
+
+		/** Frame we are waiting on to be captured */
+		TSharedPtr<class FCaptureFrame> CapturingFrame;
+	};
+
+	/** 
+	 * Class dealing with sync points. Goes over each pending capture frame data
+	 * and wait until it's done capturing. Notifies implementation from render thread or watcher's
+	 * own thread if supported.
+	 */
+	class FSyncPointWatcher : public FRunnable
+	{
+	public:
+		FSyncPointWatcher(UMediaCapture* InMediaCapture);
+		virtual ~FSyncPointWatcher();
+
+		//~ Begin FRunnable interface
+		virtual uint32 Run() override;
+		virtual void Stop() override;
+		//~ End FRunnable interface
+
+		/** Wait for all pending tasks to be completed */
+		void WaitForAllPendingTasksToComplete();
+
+		/** Waits for one pending task to be completed in order to have one free capture frame */
+		void WaitForSinglePendingTaskToComplete();
+
+		/** Queue work to be processed by our thread */
+		void QueuePendingCapture(FPendingCaptureData NewData);
+
+	private:
+
+		/** Process pending frame being captured */
+		void ProcessPendingCapture(const FPendingCaptureData& Data);
+		
+	protected:
+
+		/** Thread running sync point watch */
+		TUniquePtr<FRunnableThread> WorkingThread;
+
+		/** Whether watcher is enabled or not */
+		std::atomic<bool> bIsEnabled = false;
+
+		/** Number of pending frames to process */
+		std::atomic<uint32> PendingCaptureFramesCount = 0;
+
+		/** Event used to wake up / wait for new work */
+		FEventRef PendingTaskSignal;
+
+		/** Media capture linked to this watcher */
+		UMediaCapture* MediaCapture = nullptr;
+
+		/** Queue of work */
+		TSpscQueue<FPendingCaptureData> PendingCaptureFrames;
+	};
+}
diff --git a/Engine/Plugins/Media/MediaIOFramework/Source/MediaIOCore/Public/MediaCapture.h b/Engine/Plugins/Media/MediaIOFramework/Source/MediaIOCore/Public/MediaCapture.h
--- a/Engine/Plugins/Media/MediaIOFramework/Source/MediaIOCore/Public/MediaCapture.h
+++ b/Engine/Plugins/Media/MediaIOFramework/Source/MediaIOCore/Public/MediaCapture.h
@@ -41,6 +41,8 @@ namespace UE::MediaCaptureData
 	struct FCaptureFrameArgs;
 	class FTextureCaptureFrame;
 	class FBufferCaptureFrame;
+	class FSyncPointWatcher;
+	struct FPendingCaptureData;
 }
 
 namespace UE::MediaCapture::Private
@@ -455,6 +457,7 @@ protected:
 
 	friend class UE::MediaCaptureData::FCaptureFrame;
 	friend class UE::MediaCaptureData::FMediaCaptureHelper;
+	friend class UE::MediaCaptureData::FSyncPointWatcher;
 	struct FCaptureBaseData
 	{
 		FTimecode SourceFrameTimecode;
@@ -463,6 +466,8 @@ protected:
 		uint32 SourceFrameNumber = 0;
 	};
 
+	
+
 	/**
 	 * Capture the data that will pass along to the callback.
 	 * @note The capture is done on the Render Thread but triggered from the Game Thread.
@@ -596,7 +601,6 @@ private:
 	
 	void WaitForSingleExperimentalSchedulingTaskToComplete();
 	void WaitForAllExperimentalSchedulingTasksToComplete();
-	void CleanupCompletedExperimentalSchedulingTasks();
 	
 	void CaptureImmediate_RenderThread(const UE::MediaCaptureData::FCaptureFrameArgs& Args);
 	// Capture pipeline stuff
@@ -697,6 +701,7 @@ private:
 	static constexpr int32 MaxCaptureDataAgeInFrames = 20;
 
 	/** Structure holding data used for synchronization of buffer output */
+	friend struct UE::MediaCaptureData::FPendingCaptureData;
 	struct FMediaCaptureSyncData
 	{
 		/**
@@ -713,20 +718,14 @@ private:
 	/** Array of sync handlers (fence) to sync when captured buffer is completed */
 	TArray<TSharedPtr<FMediaCaptureSyncData>> SyncHandlers;
 
-	/** Critical section used to synchronize access to the pending readback tasks accross threads. */
-	FCriticalSection PendingReadbackTasksCriticalSection;
-
-	/** List of pending readback tasks, when CVarMediaIOEnableExperimentalScheduling and CVarMediaIOScheduleOnAnyThread are set to true. */
-	TArray<UE::Tasks::FTask> PendingReadbackTasks;
+	/** Watcher thread looking after end of capturing frames */
+	TPimplPtr<UE::MediaCaptureData::FSyncPointWatcher> SyncPointWatcher;
 
 	TSharedPtr<UE::MediaCapture::Private::FCaptureSource> CaptureSource;
 
 	/** Holds a view extension that callbacks  */
 	TSharedPtr<class FMediaCaptureSceneViewExtension> ViewExtension;
 
-	/** Helper to ensure that the async processing of frames happens in the expected order */
-	UE::MediaIO::FOrderedAsyncGate OrderedAsyncGateCaptureReady;
-
 	/** Holds the render passes that a texture will go through during media capture. */
 	TPimplPtr<UE::MediaCapture::FRenderPipeline> CaptureRenderPipeline;
 
