[MediaIO] Just-In-Time Rendering for MediaIO framework (AJA and BlackMagic)

diff --git a/Engine/Plugins/Media/AjaMedia/Source/AjaMedia/Private/Assets/AjaMediaSource.cpp b/Engine/Plugins/Media/AjaMedia/Source/AjaMedia/Private/Assets/AjaMediaSource.cpp
--- a/Engine/Plugins/Media/AjaMedia/Source/AjaMedia/Private/Assets/AjaMediaSource.cpp
+++ b/Engine/Plugins/Media/AjaMedia/Source/AjaMedia/Private/Assets/AjaMediaSource.cpp
@@ -260,6 +260,29 @@ bool UAjaMediaSource::Validate() const
 		}
 	}
 
+	if (EvaluationType == EMediaIOSampleEvaluationType::Timecode && (!bUseTimeSynchronization || AutoDetectableTimecodeFormat == EMediaIOAutoDetectableTimecodeFormat::None))
+	{
+		UE_LOG(LogAjaMedia, Warning, TEXT("The MediaSource '%s' uses 'Timecode' evaluation type which requires time synchronization and timecode enabled."), *GetName());
+		return false;
+	}
+
+	if (bFramelock && (!bRenderJIT || !bUseTimeSynchronization || AutoDetectableTimecodeFormat == EMediaIOAutoDetectableTimecodeFormat::None))
+	{
+		UE_LOG(LogAjaMedia, Warning, TEXT("The MediaSource '%s' uses 'Framelock' which requires JIT rendering, time synchronization and timecode enabled."), *GetName());
+		return false;
+	}
+
+	if (!bRenderJIT && EvaluationType == EMediaIOSampleEvaluationType::Latest)
+	{
+		UE_LOG(LogAjaMedia, Warning, TEXT("The MediaSource '%s' uses 'Latest' evaluation type which requires JIT rendering."), *GetName());
+		return false;
+	}
+
+	if (bFramelock)
+	{
+		UE_LOG(LogAjaMedia, Warning, TEXT("The MediaSource '%s' uses 'Framelock' which has not been implemented yet. This option will be ignored."), *GetName());
+	}
+
 	return true;
 }
 
@@ -292,6 +315,23 @@ void UAjaMediaSource::PostEditChangeChainProperty(struct FPropertyChangedChainEv
 		{
 			bUseTimeSynchronization = false;
 			bEncodeTimecodeInTexel = false;
+			EvaluationType = EMediaIOSampleEvaluationType::PlatformTime;
+			bFramelock = false;
+		}
+	}
+
+	if (PropertyChangedEvent.Property->GetFName() == GET_MEMBER_NAME_CHECKED(UCaptureCardMediaSource, EvaluationType))
+	{
+		// 'Timecode' evaluation is not allowed if no timecode set
+		if (AutoDetectableTimecodeFormat == EMediaIOAutoDetectableTimecodeFormat::None || !bUseTimeSynchronization)
+		{
+			EvaluationType = (EvaluationType == EMediaIOSampleEvaluationType::Timecode ? EMediaIOSampleEvaluationType::PlatformTime : EvaluationType);
+		}
+
+		// 'Latest' evaluation is available in JITR only
+		if (!bRenderJIT)
+		{
+			EvaluationType = (EvaluationType == EMediaIOSampleEvaluationType::Latest ? EMediaIOSampleEvaluationType::PlatformTime : EvaluationType);
 		}
 	}
 
diff --git a/Engine/Plugins/Media/AjaMedia/Source/AjaMedia/Private/Player/AjaMediaPlayer.cpp b/Engine/Plugins/Media/AjaMedia/Source/AjaMedia/Private/Player/AjaMediaPlayer.cpp
--- a/Engine/Plugins/Media/AjaMedia/Source/AjaMedia/Private/Player/AjaMediaPlayer.cpp
+++ b/Engine/Plugins/Media/AjaMedia/Source/AjaMedia/Private/Player/AjaMediaPlayer.cpp
@@ -62,26 +62,8 @@ FAjaMediaPlayer::FAjaMediaPlayer(IMediaEventSink& InEventSink)
 	, AudioSamplePool(MakeUnique<FAjaMediaAudioSamplePool>())
 	, MetadataSamplePool(MakeUnique<FAjaMediaBinarySamplePool>())
 	, TextureSamplePool(MakeUnique<FAjaMediaTextureSamplePool>())
-	, MaxNumAudioFrameBuffer(8)
-	, MaxNumMetadataFrameBuffer(8)
-	, MaxNumVideoFrameBuffer(8)
-	, AjaThreadNewState(EMediaState::Closed)
 	, EventSink(InEventSink)
-	, AjaThreadAudioChannels(0)
-	, AjaThreadAudioSampleRate(0)
-	, AjaThreadFrameDropCount(0)
-	, LastFrameDropCount(0)
-	, PreviousFrameDropCount(0)
-	, bEncodeTimecodeInTexel(false)
-	, bUseFrameTimecode(false)
-	, bIsSRGBInput(false)
-	, bUseAncillary(false)
-	, bUseAudio(false)
-	, bUseVideo(false)
-	, bVerifyFrameDropCount(true)
-	, InputChannel(nullptr)
 	, SupportedSampleTypes(EMediaIOSampleType::None)
-	, bPauseRequested(false)
 	, DeviceProvider(MakePimpl<FAjaDeviceProvider>())
 {
 }
@@ -159,9 +141,6 @@ void FAjaMediaPlayer::Close()
 	MetadataSamplePool->Reset();
 	TextureSamplePool->Reset();
 
-	UnregisterSampleBuffers();
-	UnregisterTextures();
-
 	//Disable all our channels from the monitor
 	Samples->EnableTimedDataChannels(this, EMediaIOSampleType::None);
 
@@ -306,40 +285,45 @@ void FAjaMediaPlayer::VerifyFrameDropCount()
 		}
 		LastFrameDropCount = FrameDropCount;
 
-		const int32 CurrentMetadataDropCount = Samples->GetMetadataFrameDropCount();
-		int32 DeltaMetadataDropCount = CurrentMetadataDropCount;
-		if (CurrentMetadataDropCount >= PreviousMetadataFrameDropCount)
-		{
-			DeltaMetadataDropCount = CurrentMetadataDropCount - PreviousMetadataFrameDropCount;
-		}
-		PreviousMetadataFrameDropCount = CurrentMetadataDropCount;
-		if (DeltaMetadataDropCount > 0)
+		// JITR pipeline always keeps the sample pool full. So every new sample increments internal drop count.
+		// This if-condition is intended to avoid "Lost %d XXX frames on input..." message spam every frame.
+		if (!IsJustInTimeRenderingEnabled())
 		{
-			UE_LOG(LogAjaMedia, Warning, TEXT("Lost %d metadata frames on input %s. Frame rate is either too slow or buffering capacity is too small."), DeltaMetadataDropCount, *GetUrl());
-		}
+			const int32 CurrentMetadataDropCount = Samples->GetMetadataFrameDropCount();
+			int32 DeltaMetadataDropCount = CurrentMetadataDropCount;
+			if (CurrentMetadataDropCount >= PreviousMetadataFrameDropCount)
+			{
+				DeltaMetadataDropCount = CurrentMetadataDropCount - PreviousMetadataFrameDropCount;
+			}
+			PreviousMetadataFrameDropCount = CurrentMetadataDropCount;
+			if (DeltaMetadataDropCount > 0)
+			{
+				UE_LOG(LogAjaMedia, Warning, TEXT("Lost %d metadata frames on input %s. Frame rate is either too slow or buffering capacity is too small."), DeltaMetadataDropCount, *GetUrl());
+			}
 
-		const int32 CurrentAudioDropCount = Samples->GetAudioFrameDropCount();
-		int32 DeltaAudioDropCount = CurrentAudioDropCount;
-		if (CurrentAudioDropCount >= PreviousAudioFrameDropCount)
-		{
-			DeltaAudioDropCount = CurrentAudioDropCount - PreviousAudioFrameDropCount;
-		}
-		PreviousAudioFrameDropCount = CurrentAudioDropCount;
-		if (DeltaAudioDropCount > 0)
-		{
-			UE_LOG(LogAjaMedia, Warning, TEXT("Lost %d audio frames on input %s. Frame rate is either too slow or buffering capacity is too small."), DeltaAudioDropCount, *GetUrl());
-		}
+			const int32 CurrentAudioDropCount = Samples->GetAudioFrameDropCount();
+			int32 DeltaAudioDropCount = CurrentAudioDropCount;
+			if (CurrentAudioDropCount >= PreviousAudioFrameDropCount)
+			{
+				DeltaAudioDropCount = CurrentAudioDropCount - PreviousAudioFrameDropCount;
+			}
+			PreviousAudioFrameDropCount = CurrentAudioDropCount;
+			if (DeltaAudioDropCount > 0)
+			{
+				UE_LOG(LogAjaMedia, Warning, TEXT("Lost %d audio frames on input %s. Frame rate is either too slow or buffering capacity is too small."), DeltaAudioDropCount, *GetUrl());
+			}
 
-		const int32 CurrentVideoDropCount = Samples->GetVideoFrameDropCount();
-		int32 DeltaVideoDropCount = CurrentVideoDropCount;
-		if (CurrentVideoDropCount >= PreviousVideoFrameDropCount)
-		{
-			DeltaVideoDropCount = CurrentVideoDropCount - PreviousVideoFrameDropCount;
-		}
-		PreviousVideoFrameDropCount = CurrentVideoDropCount;
-		if (DeltaVideoDropCount > 0)
-		{
-			UE_LOG(LogAjaMedia, Warning, TEXT("Lost %d video frames on input %s. Frame rate is either too slow or buffering capacity is too small."), DeltaVideoDropCount, *GetUrl());
+			const int32 CurrentVideoDropCount = Samples->GetVideoFrameDropCount();
+			int32 DeltaVideoDropCount = CurrentVideoDropCount;
+			if (CurrentVideoDropCount >= PreviousVideoFrameDropCount)
+			{
+				DeltaVideoDropCount = CurrentVideoDropCount - PreviousVideoFrameDropCount;
+			}
+			PreviousVideoFrameDropCount = CurrentVideoDropCount;
+			if (DeltaVideoDropCount > 0)
+			{
+				UE_LOG(LogAjaMedia, Warning, TEXT("Lost %d video frames on input %s. Frame rate is either too slow or buffering capacity is too small."), DeltaVideoDropCount, *GetUrl());
+			}
 		}
 	}
 }
@@ -428,7 +412,7 @@ bool FAjaMediaPlayer::OnRequestInputBuffer(const AJA::AJARequestInputBufferData&
 	if (bUseAncillary && InRequestBuffer.AncF2BufferSize > 0)
 	{
 		AjaThreadCurrentAncF2Sample = MetadataSamplePool->AcquireShared();
-		OutRequestedBuffer.AncBuffer = reinterpret_cast<uint8_t*>(AjaThreadCurrentAncF2Sample->RequestBuffer(InRequestBuffer.AncF2BufferSize));
+		OutRequestedBuffer.AncF2Buffer = reinterpret_cast<uint8_t*>(AjaThreadCurrentAncF2Sample->RequestBuffer(InRequestBuffer.AncF2BufferSize));
 	}
 
 	// Audio
@@ -441,11 +425,13 @@ bool FAjaMediaPlayer::OnRequestInputBuffer(const AJA::AJARequestInputBufferData&
 	// Video
 	if (bUseVideo && InRequestBuffer.VideoBufferSize > 0 && InRequestBuffer.bIsProgressivePicture)
 	{
-		const bool bCanUseGPUTextureTransfer = CanUseGPUTextureTransfer();
-
-		if (bCanUseGPUTextureTransfer)
+		// JITR pipeline uses a single texture for GPU texture transfer as a shared one, so it's always available.
+		const bool bIsJustInTimeRenderingEnabled = IsJustInTimeRenderingEnabled();
+		const bool bCanUseGPUTextureTransfer     = CanUseGPUTextureTransfer();
+		if (!bIsJustInTimeRenderingEnabled && bCanUseGPUTextureTransfer)
 		{
-			if (!Textures.Num())
+			const bool bHasTextureAvailableForGPUTransfer = HasTextureAvailableForGPUTransfer();
+			if (!bHasTextureAvailableForGPUTransfer)
 			{
 				UE_LOG(LogAjaMedia, Error, TEXT("No texture available while doing a gpu texture transfer."));
 				return false;
@@ -455,9 +441,9 @@ bool FAjaMediaPlayer::OnRequestInputBuffer(const AJA::AJARequestInputBufferData&
 		AjaThreadCurrentTextureSample = TextureSamplePool->AcquireShared();
 		AjaThreadCurrentTextureSample->RequestBuffer(InRequestBuffer.VideoBufferSize);
 
-		if (bCanUseGPUTextureTransfer)
+		if (bIsJustInTimeRenderingEnabled && bCanUseGPUTextureTransfer)
 		{
-			PreGPUTransfer(AjaThreadCurrentTextureSample);
+			AjaThreadCurrentTextureSample->SetAwaitingForGPUTransfer();
 		}
 
 		OutRequestedBuffer.VideoBuffer = reinterpret_cast<uint8_t*>(AjaThreadCurrentTextureSample->GetMutableBuffer());
@@ -500,7 +486,7 @@ bool FAjaMediaPlayer::OnInputFrameReceived(const AJA::AJAInputFrameData& InInput
 		const FFrameNumber ConvertedFrameNumber = DecodedTimecode.GetValue().ToFrameNumber(VideoFrameRate);
 		const double NumberOfSeconds = ConvertedFrameNumber.Value * VideoFrameRate.AsInterval();
 		const FTimespan TimecodeDecodedTime = FTimespan::FromSeconds(NumberOfSeconds);
-		if (bUseTimeSynchronization)
+		if (EvaluationType == EMediaIOSampleEvaluationType::Timecode)
 		{
 			DecodedTime = TimecodeDecodedTime;
 			DecodedTimeF2 = TimecodeDecodedTime + FTimespan::FromSeconds(VideoFrameRate.AsInterval());
@@ -522,7 +508,7 @@ bool FAjaMediaPlayer::OnInputFrameReceived(const AJA::AJAInputFrameData& InInput
 		{
 			if (AjaThreadCurrentAncSample->SetProperties(DecodedTime, VideoFrameRate, DecodedTimecode))
 			{
-				Samples->AddMetadata(AjaThreadCurrentAncSample.ToSharedRef());
+				AddMetadataSample(AjaThreadCurrentAncSample.ToSharedRef());
 			}
 		}
 		else
@@ -530,7 +516,7 @@ bool FAjaMediaPlayer::OnInputFrameReceived(const AJA::AJAInputFrameData& InInput
 			auto MetaDataSample = MetadataSamplePool->AcquireShared();
 			if (MetaDataSample->Initialize(InAncillaryFrame.AncBuffer, InAncillaryFrame.AncBufferSize, DecodedTime, VideoFrameRate, DecodedTimecode))
 			{
-				Samples->AddMetadata(MetaDataSample);
+				AddMetadataSample(MetaDataSample);
 			}
 		}
 	}
@@ -542,7 +528,7 @@ bool FAjaMediaPlayer::OnInputFrameReceived(const AJA::AJAInputFrameData& InInput
 		{
 			if (AjaThreadCurrentAncF2Sample->SetProperties(DecodedTimeF2, VideoFrameRate, DecodedTimecodeF2))
 			{
-				Samples->AddMetadata(AjaThreadCurrentAncF2Sample.ToSharedRef());
+				AddMetadataSample(AjaThreadCurrentAncF2Sample.ToSharedRef());
 			}
 		}
 		else
@@ -550,7 +536,7 @@ bool FAjaMediaPlayer::OnInputFrameReceived(const AJA::AJAInputFrameData& InInput
 			auto MetaDataSample = MetadataSamplePool->AcquireShared();
 			if (MetaDataSample->Initialize(InAncillaryFrame.AncF2Buffer, InAncillaryFrame.AncF2BufferSize, DecodedTimeF2, VideoFrameRate, DecodedTimecodeF2))
 			{
-				Samples->AddMetadata(MetaDataSample);
+				AddMetadataSample(MetaDataSample);
 			}
 		}
 	}
@@ -562,7 +548,7 @@ bool FAjaMediaPlayer::OnInputFrameReceived(const AJA::AJAInputFrameData& InInput
 		{
 			if (AjaThreadCurrentAudioSample->SetProperties(InAudioFrame.AudioBufferSize / sizeof(int32), InAudioFrame.NumChannels, InAudioFrame.AudioRate, DecodedTime, DecodedTimecode))
 			{
-				Samples->AddAudio(AjaThreadCurrentAudioSample.ToSharedRef());
+				AddAudioSample(AjaThreadCurrentAudioSample.ToSharedRef());
 			}
 
 			AjaThreadAudioChannels = AjaThreadCurrentAudioSample->GetChannels();
@@ -573,7 +559,7 @@ bool FAjaMediaPlayer::OnInputFrameReceived(const AJA::AJAInputFrameData& InInput
 			auto AudioSample = AudioSamplePool->AcquireShared();
 			if (AudioSample->Initialize(InAudioFrame, DecodedTime, DecodedTimecode))
 			{
-				Samples->AddAudio(AudioSample);
+				AddAudioSample(AudioSample);
 			}
 
 			AjaThreadAudioChannels = AudioSample->GetChannels();
@@ -685,14 +671,7 @@ bool FAjaMediaPlayer::OnInputFrameReceived(const AJA::AJAInputFrameData& InInput
 		{
 			if (AjaThreadCurrentTextureSample->SetProperties(InVideoFrame.Stride, InVideoFrame.Width, InVideoFrame.Height, VideoSampleFormat, DecodedTime, VideoFrameRate, DecodedTimecode, ColorFormat))
 			{
-				if (CanUseGPUTextureTransfer())
-				{
-					ExecuteGPUTransfer(AjaThreadCurrentTextureSample);
-				}
-				else
-				{
-					Samples->AddVideo(AjaThreadCurrentTextureSample.ToSharedRef());
-				}
+				AddVideoSample(AjaThreadCurrentTextureSample.ToSharedRef());
 			}
 		}
 		else
@@ -720,7 +699,7 @@ bool FAjaMediaPlayer::OnInputFrameReceived(const AJA::AJAInputFrameData& InInput
 					auto TextureSample = TextureSamplePool->AcquireShared();
 					if (TextureSample->InitializeProgressive(InVideoFrame, VideoSampleFormat, DecodedTime, VideoFrameRate, DecodedTimecode, ColorFormat))
 					{
-						Samples->AddVideo(TextureSample);
+						AddVideoSample(TextureSample);
 					}
 				}
 				else
@@ -748,7 +727,7 @@ bool FAjaMediaPlayer::OnInputFrameReceived(const AJA::AJAInputFrameData& InInput
 
 					for (const TSharedRef<FMediaIOCoreTextureSampleBase>& TextureSample : DeinterlacedSamples)
 					{
-						Samples->AddVideo(TextureSample);
+						AddVideoSample(TextureSample);
 					}
 				}
 			}
@@ -790,11 +769,6 @@ void FAjaMediaPlayer::SetupSampleChannels()
 	Samples->InitializeMetadataBuffer(MetadataSettings);
 }
 
-void FAjaMediaPlayer::AddVideoSample(const TSharedRef<FMediaIOCoreTextureSampleBase>& InSample)
-{
-	Samples->AddVideo(InSample);
-}
-
 bool FAjaMediaPlayer::Open_Internal(const FString& Url, const IMediaOptions* Options, AJA::AJAInputOutputChannelOptions AjaOptions)
 {
 	AJA::AJADeviceOptions DeviceOptions(Options->GetMediaOption(AjaMediaOption::DeviceIndex, (int64)0));
@@ -1012,7 +986,7 @@ void FAjaMediaPlayer::OnAutoDetected(TArray<FAjaDeviceProvider::FMediaIOConfigur
 			break;
 		}
 	}
-			
+
 	if (!bConfigurationFound)
 	{
 		UE_LOG(LogAjaMedia, Warning, TEXT("No configuration was detected for MediaPlayer."));
@@ -1021,6 +995,8 @@ void FAjaMediaPlayer::OnAutoDetected(TArray<FAjaDeviceProvider::FMediaIOConfigur
 	}
 
 	Open_Internal(Url, Options, AjaOptions);
+
+	NotifyVideoFormatDetected();
 }
 
 bool FAjaMediaPlayer::SetRate(float Rate)
@@ -1040,7 +1016,7 @@ bool FAjaMediaPlayer::SetRate(float Rate)
 	return false;
 }
 
-TSharedPtr<FMediaIOCoreTextureSampleBase> FAjaMediaPlayer::AcquireSample_AnyThread() const
+TSharedPtr<FMediaIOCoreTextureSampleBase> FAjaMediaPlayer::AcquireTextureSample_AnyThread() const
 {
 	return TextureSamplePool->AcquireShared();
 }
diff --git a/Engine/Plugins/Media/AjaMedia/Source/AjaMedia/Private/Player/AjaMediaPlayer.h b/Engine/Plugins/Media/AjaMedia/Source/AjaMedia/Private/Player/AjaMediaPlayer.h
--- a/Engine/Plugins/Media/AjaMedia/Source/AjaMedia/Private/Player/AjaMediaPlayer.h
+++ b/Engine/Plugins/Media/AjaMedia/Source/AjaMedia/Private/Player/AjaMediaPlayer.h
@@ -115,8 +115,8 @@ protected:
 	{
 		return AjaColorFormat == EAjaMediaSourceColorFormat::YUV2_8bit ? EMediaIOCoreColorFormat::YUV8 : EMediaIOCoreColorFormat::YUV10;
 	}
-	virtual void AddVideoSample(const TSharedRef<FMediaIOCoreTextureSampleBase>& InSample) override;
-	virtual TSharedPtr<FMediaIOCoreTextureSampleBase> AcquireSample_AnyThread() const override;
+
+	virtual TSharedPtr<FMediaIOCoreTextureSampleBase> AcquireTextureSample_AnyThread() const override;
 	//~ End FMediaIOCorePlayerBase interface
 
 private:
@@ -137,49 +137,49 @@ private:
 	TSharedPtr<FAjaMediaTextureSample, ESPMode::ThreadSafe> AjaThreadCurrentTextureSample;
 
 	/** The media sample cache. */
-	int32 MaxNumAudioFrameBuffer;
-	int32 MaxNumMetadataFrameBuffer;
-	int32 MaxNumVideoFrameBuffer;
+	int32 MaxNumAudioFrameBuffer    = 0;
+	int32 MaxNumMetadataFrameBuffer = 0;
+	int32 MaxNumVideoFrameBuffer    = 0;
 
 	/** Current state of the media player. */
-	EMediaState AjaThreadNewState;
+	EMediaState AjaThreadNewState = EMediaState::Closed;
 
 	/** The media event handler. */
 	IMediaEventSink& EventSink;
 
 	/** Number of audio channels in the last received sample. */
-	int32 AjaThreadAudioChannels;
+	int32 AjaThreadAudioChannels = 0;
 
 	/** Audio sample rate in the last received sample. */
-	int32 AjaThreadAudioSampleRate;
+	int32 AjaThreadAudioSampleRate = 0;
 
 	/** Number of frames drop from the last tick. */
-	int32 AjaThreadFrameDropCount;
-	int32 PreviousMetadataFrameDropCount;
-	int32 PreviousAudioFrameDropCount;
-	int32 PreviousVideoFrameDropCount;
+	int32 AjaThreadFrameDropCount = 0;
+	int32 PreviousMetadataFrameDropCount = 0;
+	int32 PreviousAudioFrameDropCount = 0;
+	int32 PreviousVideoFrameDropCount = 0;
 
 	/** Number of frames drop from the last tick. */
-	uint32 LastFrameDropCount;
-	uint32 PreviousFrameDropCount;
+	uint32 LastFrameDropCount     = 0;
+	uint32 PreviousFrameDropCount = 0;
 
 	/** Whether to use the time code embedded in AJA frames. */
-	bool bEncodeTimecodeInTexel;
+	bool bEncodeTimecodeInTexel = false;
 
 	/** Whether to use the timecode embedded in a frame. */
-	bool bUseFrameTimecode;
+	bool bUseFrameTimecode = false;
 
 	/** Whether the input is in sRGB and can have a ToLinear conversion. */
-	bool bIsSRGBInput;
+	bool bIsSRGBInput = false;
 
 	/** Which field need to be capture. */
-	bool bUseAncillary;
-	bool bUseAudio;
-	bool bUseVideo;
-	bool bVerifyFrameDropCount;
+	bool bUseAncillary = false;
+	bool bUseAudio = false;
+	bool bUseVideo = false;
+	bool bVerifyFrameDropCount = true;
 
 	/** Maps to the current input Device */
-	AJA::AJAInputChannel* InputChannel;
+	AJA::AJAInputChannel* InputChannel = nullptr;
 
 	/** Used to flag which sample types we advertise as supported for timed data monitoring */
 	EMediaIOSampleType SupportedSampleTypes;
@@ -188,7 +188,7 @@ private:
 	AJA::FTimecode AjaThreadPreviousFrameTimecode;
 
 	/** Flag to indicate that pause is being requested */
-	std::atomic<bool> bPauseRequested;
+	std::atomic<bool> bPauseRequested = false;
 
 	EAjaMediaSourceColorFormat AjaColorFormat = EAjaMediaSourceColorFormat::YUV2_8bit;
 
diff --git a/Engine/Plugins/Media/BlackmagicMedia/Source/BlackmagicCore/Private/BlackmagicInputChannel.cpp b/Engine/Plugins/Media/BlackmagicMedia/Source/BlackmagicCore/Private/BlackmagicInputChannel.cpp
--- a/Engine/Plugins/Media/BlackmagicMedia/Source/BlackmagicCore/Private/BlackmagicInputChannel.cpp
+++ b/Engine/Plugins/Media/BlackmagicMedia/Source/BlackmagicCore/Private/BlackmagicInputChannel.cpp
@@ -29,6 +29,32 @@ namespace BlackmagicDesign
 		const uint32_t DefaultAudioSampleRate = static_cast<uint32_t>(bmdDefaultAudioSampleRate);
 		const BMDAudioSampleType DefaultAudioSampleType = bmdAudioSampleType32bitInteger;
 
+		/** Implements buffer guarding */
+		struct FSampleBufferHolder : public BlackmagicDesign::IInputEventCallback::FFrameBufferHolder
+		{
+		public:
+			FSampleBufferHolder(IUnknown* InObject)
+				: Object(InObject)
+			{
+				if (Object)
+				{
+					Object->AddRef();
+				}
+			}
+
+			~FSampleBufferHolder()
+			{
+				if (Object)
+				{
+					Object->Release();
+				}
+			}
+
+		private:
+			IUnknown* Object = nullptr;
+		};
+
+
 		FInputChannelNotificationCallback::FInputChannelNotificationCallback(FInputChannel* InOwner, IDeckLinkInput* InDeckLinkInput, IDeckLinkStatus* InDeckLinkStatus)
 			: InputChannel(InOwner)
 			, DeckLinkInput(InDeckLinkInput)
@@ -286,6 +312,23 @@ namespace BlackmagicDesign
 				for (FInputChannel::FListener& Listener : InputChannel->Listeners)
 				{
 					Listener.Callback->OnFrameReceived(FrameInfo);
+
+					{
+						IInputEventCallback::FFrameReceivedBufferHolders BufferHolders;
+						Listener.Callback->OnFrameReceived(FrameInfo, BufferHolders);
+
+						// Instantiate audio buffer holder if requested
+						if (BufferHolders.AudioBufferHolder)
+						{
+							*BufferHolders.AudioBufferHolder = MakeShared<BlackmagicDesign::Private::FSampleBufferHolder>(InAudioPacket);
+						}
+
+						// Instantiate video buffer holder if requested
+						if (BufferHolders.VideoBufferHolder)
+						{
+							*BufferHolders.VideoBufferHolder = MakeShared<BlackmagicDesign::Private::FSampleBufferHolder>(InVideoFrame);
+						}
+					}
 				}
 			}
 
diff --git a/Engine/Plugins/Media/BlackmagicMedia/Source/BlackmagicCore/Public/BlackmagicLib.h b/Engine/Plugins/Media/BlackmagicMedia/Source/BlackmagicCore/Public/BlackmagicLib.h
--- a/Engine/Plugins/Media/BlackmagicMedia/Source/BlackmagicCore/Public/BlackmagicLib.h
+++ b/Engine/Plugins/Media/BlackmagicMedia/Source/BlackmagicCore/Public/BlackmagicLib.h
@@ -301,6 +301,22 @@ namespace BlackmagicDesign
 			int32_t AudioRate;
 		};
 
+		/** An empty type to expose. The real job is done in an internal child implementation. */
+		struct BLACKMAGICCORE_API FFrameBufferHolder
+		{
+			virtual ~FFrameBufferHolder() = default;
+		};
+
+		/**
+		 * A container for video and audio buffer holders. These holders prevent
+		 * internal buffers from being released while they are in use.
+		 */
+		struct BLACKMAGICCORE_API FFrameReceivedBufferHolders
+		{
+			TSharedPtr<FFrameBufferHolder>* VideoBufferHolder = nullptr;
+			TSharedPtr<FFrameBufferHolder>* AudioBufferHolder = nullptr;
+		};
+
 		virtual ~IInputEventCallback();
 
 		virtual void AddRef() = 0;
@@ -309,7 +325,12 @@ namespace BlackmagicDesign
 		virtual void OnInitializationCompleted(bool bSuccess) = 0;
 		virtual void OnShutdownCompleted() = 0;
 
-		virtual void OnFrameReceived(const FFrameReceivedInfo&) = 0;
+		virtual void OnFrameReceived(const FFrameReceivedInfo&)
+		{ }
+
+		virtual void OnFrameReceived(const FFrameReceivedInfo&, FFrameReceivedBufferHolders& /* out */)
+		{ }
+
 		virtual void OnFrameFormatChanged(const FFormatInfo& NewFormat) = 0;
 		virtual void OnInterlacedOddFieldEvent(long long FrameNumber) = 0;
 	};
diff --git a/Engine/Plugins/Media/BlackmagicMedia/Source/BlackmagicMedia/Private/Assets/BlackmagicMediaSource.cpp b/Engine/Plugins/Media/BlackmagicMedia/Source/BlackmagicMedia/Private/Assets/BlackmagicMediaSource.cpp
--- a/Engine/Plugins/Media/BlackmagicMedia/Source/BlackmagicMedia/Private/Assets/BlackmagicMediaSource.cpp
+++ b/Engine/Plugins/Media/BlackmagicMedia/Source/BlackmagicMedia/Private/Assets/BlackmagicMediaSource.cpp
@@ -154,6 +154,30 @@ bool UBlackmagicMediaSource::Validate() const
 		return false;
 	}
 
+
+	if (EvaluationType == EMediaIOSampleEvaluationType::Timecode && (!bUseTimeSynchronization || AutoDetectableTimecodeFormat == EMediaIOAutoDetectableTimecodeFormat::None))
+	{
+		UE_LOG(LogBlackmagicMedia, Warning, TEXT("The MediaSource '%s' uses 'Timecode' evaluation type which requires time synchronization and timecode enabled."), *GetName());
+		return false;
+	}
+
+	if (bFramelock && (!bRenderJIT || !bUseTimeSynchronization || AutoDetectableTimecodeFormat == EMediaIOAutoDetectableTimecodeFormat::None))
+	{
+		UE_LOG(LogBlackmagicMedia, Warning, TEXT("The MediaSource '%s' uses 'Framelock' which requires JIT rendering, time synchronization and timecode enabled."), *GetName());
+		return false;
+	}
+
+	if (!bRenderJIT && EvaluationType == EMediaIOSampleEvaluationType::Latest)
+	{
+		UE_LOG(LogBlackmagicMedia, Warning, TEXT("The MediaSource '%s' uses 'Latest' evaluation type which requires JIT rendering."), *GetName());
+		return false;
+	}
+
+	if (bFramelock)
+	{
+		UE_LOG(LogBlackmagicMedia, Warning, TEXT("The MediaSource '%s' uses 'Framelock' which has not been implemented yet. This option will be ignored."), *GetName());
+	}
+
 	return true;
 }
 
@@ -186,6 +210,23 @@ void UBlackmagicMediaSource::PostEditChangeChainProperty(struct FPropertyChanged
 		{
 			bUseTimeSynchronization = false;
 			bEncodeTimecodeInTexel = false;
+			EvaluationType = EMediaIOSampleEvaluationType::PlatformTime;
+			bFramelock = false;
+		}
+	}
+
+	if (InPropertyChangedEvent.Property->GetFName() == GET_MEMBER_NAME_CHECKED(UCaptureCardMediaSource, EvaluationType))
+	{
+		// 'Timecode' evaluation is not allowed if no timecode set
+		if (AutoDetectableTimecodeFormat == EMediaIOAutoDetectableTimecodeFormat::None)
+		{
+			EvaluationType = (EvaluationType == EMediaIOSampleEvaluationType::Timecode ? EMediaIOSampleEvaluationType::PlatformTime : EvaluationType);
+		}
+
+		// 'Latest' evaluation is available in JITR only
+		if (!bRenderJIT)
+		{
+			EvaluationType = (EvaluationType == EMediaIOSampleEvaluationType::Latest ? EMediaIOSampleEvaluationType::PlatformTime : EvaluationType);
 		}
 	}
 
diff --git a/Engine/Plugins/Media/BlackmagicMedia/Source/BlackmagicMedia/Private/Player/BlackmagicMediaPlayer.cpp b/Engine/Plugins/Media/BlackmagicMedia/Source/BlackmagicMedia/Private/Player/BlackmagicMediaPlayer.cpp
--- a/Engine/Plugins/Media/BlackmagicMedia/Source/BlackmagicMedia/Private/Player/BlackmagicMediaPlayer.cpp
+++ b/Engine/Plugins/Media/BlackmagicMedia/Source/BlackmagicMedia/Private/Player/BlackmagicMediaPlayer.cpp
@@ -2,9 +2,7 @@
 
 #include "Player/BlackmagicMediaPlayer.h"
 
-#include "BlackmagicMediaPrivate.h"
 #include "BlackmagicMediaDefinitions.h"
-#include "BlackmagicMediaSource.h"
 
 #include "Engine/GameEngine.h"
 #include "HAL/CriticalSection.h"
@@ -26,8 +24,6 @@
 #include "Styling/SlateStyle.h"
 #include "Templates/Atomic.h"
 
-#include "BlackmagicMediaSource.h"
-
 #if WITH_EDITOR
 #include "EngineAnalytics.h"
 #endif
@@ -125,16 +121,21 @@ namespace BlackmagicMediaPlayerHelpers
 					UE_LOG(LogBlackmagicMedia, Warning, TEXT("Lost %d audio frames on input %s. Frame rate is either too slow or buffering capacity is too small."), DeltaAudioDropCount, *InUrl);
 				}
 
-				const int32 CurrentVideoDropCount = MediaPlayer->Samples->GetVideoFrameDropCount();
-				int32 DeltaVideoDropCount = CurrentVideoDropCount;
-				if (CurrentVideoDropCount >= PreviousVideoFrameDropCount)
-				{
-					DeltaVideoDropCount = CurrentVideoDropCount - PreviousVideoFrameDropCount;
-				}
-				PreviousVideoFrameDropCount = CurrentVideoDropCount;
-				if (DeltaVideoDropCount > 0)
+				// JITR pipeline always keeps the sample pool full. So every new sample increments internal drop count.
+				// This if-condition is intended to avoid "Lost %d XXX frames on input..." message spam every frame.
+				if (!MediaPlayer->IsJustInTimeRenderingEnabled())
 				{
-					UE_LOG(LogBlackmagicMedia, Warning, TEXT("Lost %d video frames on input %s. Frame rate is either too slow or buffering capacity is too small."), DeltaVideoDropCount, *InUrl);
+					const int32 CurrentVideoDropCount = MediaPlayer->Samples->GetVideoFrameDropCount();
+					int32 DeltaVideoDropCount = CurrentVideoDropCount;
+					if (CurrentVideoDropCount >= PreviousVideoFrameDropCount)
+					{
+						DeltaVideoDropCount = CurrentVideoDropCount - PreviousVideoFrameDropCount;
+					}
+					PreviousVideoFrameDropCount = CurrentVideoDropCount;
+					if (DeltaVideoDropCount > 0)
+					{
+						UE_LOG(LogBlackmagicMedia, Warning, TEXT("Lost %d video frames on input %s. Frame rate is either too slow or buffering capacity is too small."), DeltaVideoDropCount, *InUrl);
+					}
 				}
 			}
 		}
@@ -164,7 +165,9 @@ namespace BlackmagicMediaPlayerHelpers
 			MediaState = EMediaState::Closed;
 		}
 
-		virtual void OnFrameReceived(const BlackmagicDesign::IInputEventCallback::FFrameReceivedInfo& InFrameInfo) override
+		virtual void OnFrameReceived(
+			const BlackmagicDesign::IInputEventCallback::FFrameReceivedInfo& InFrameInfo,
+			BlackmagicDesign::IInputEventCallback::FFrameReceivedBufferHolders& OutBufferHolders) override
 		{
 			SCOPE_CYCLE_COUNTER(STAT_Blackmagic_MediaPlayer_ProcessReceivedFrame);
 
@@ -228,7 +231,7 @@ namespace BlackmagicMediaPlayerHelpers
 					const FFrameNumber ConvertedFrameNumber = DecodedTimecode.GetValue().ToFrameNumber(MediaPlayer->VideoFrameRate);
 					const double NumberOfSeconds = ConvertedFrameNumber.Value * MediaPlayer->VideoFrameRate.AsInterval();
 					const FTimespan TimecodeDecodedTime = FTimespan::FromSeconds(NumberOfSeconds);
-					if (MediaPlayer->bUseTimeSynchronization)
+					if (MediaPlayer->EvaluationType == EMediaIOSampleEvaluationType::Timecode)
 					{
 						DecodedTime = TimecodeDecodedTime;
 						DecodedTimeF2 = TimecodeDecodedTime + FTimespan::FromSeconds(MediaPlayer->VideoFrameRate.AsInterval());
@@ -259,7 +262,7 @@ namespace BlackmagicMediaPlayerHelpers
 						, DecodedTime
 						, DecodedTimecode))
 					{
-						MediaPlayer->Samples->AddAudio(AudioSamle);
+						MediaPlayer->AddAudioSample(AudioSamle);
 
 						LastBitsPerSample = sizeof(int32);
 						LastSampleRate = InFrameInfo.AudioRate;
@@ -356,15 +359,11 @@ namespace BlackmagicMediaPlayerHelpers
 						bool bGPUDirectTexturesAvailable = false;
 						if (MediaPlayer->CanUseGPUTextureTransfer())
 						{
-							if (!MediaPlayer->Textures.Num())
+							bGPUDirectTexturesAvailable = MediaPlayer->HasTextureAvailableForGPUTransfer();
+							if (!bGPUDirectTexturesAvailable)
 							{
-								bGPUDirectTexturesAvailable = false;
 								UE_LOG(LogBlackmagicMedia, Error, TEXT("No texture available while doing a gpu texture transfer."));
 							}
-							else
-							{
-								bGPUDirectTexturesAvailable = true;
-							}
 						}
 
 						auto TextureSample = MediaPlayer->TextureSamplePool->AcquireShared();
@@ -395,13 +394,14 @@ namespace BlackmagicMediaPlayerHelpers
 						{
 							if (bGPUDirectTexturesAvailable && MediaPlayer->CanUseGPUTextureTransfer())
 							{
-								MediaPlayer->PreGPUTransfer(TextureSample);
-								MediaPlayer->ExecuteGPUTransfer(TextureSample);
-							}
-							else
-							{
-								MediaPlayer->Samples->AddVideo(TextureSample);
+								// Mark this sample ready for GPU transfer
+								TextureSample->SetAwaitingForGPUTransfer();
+
+								// Keep the internal buffer alive until GPU transfer is finished
+								OutBufferHolders.VideoBufferHolder = &TextureSample->GetBlackmagicInternalBufferLocker();
 							}
+
+							MediaPlayer->AddVideoSample(TextureSample);
 						}
 					}
 					else
@@ -424,7 +424,7 @@ namespace BlackmagicMediaPlayerHelpers
 
 						for (const TSharedRef<FMediaIOCoreTextureSampleBase>& TextureSample : DeinterlacedSamples)
 						{
-							MediaPlayer->Samples->AddVideo(TextureSample);
+							MediaPlayer->AddVideoSample(TextureSample);
 						}
 					}
 				}
@@ -436,6 +436,14 @@ namespace BlackmagicMediaPlayerHelpers
 			if (MediaPlayer->bAutoDetect)
 			{
 				MediaPlayer->VideoFrameRate = FFrameRate(NewFormat.FrameRateNumerator, NewFormat.FrameRateDenominator);
+				MediaPlayer->BaseSettings.FrameRate = MediaPlayer->VideoFrameRate;
+
+				MediaPlayer->VideoTrackFormat.Dim = FIntPoint(NewFormat.Width, NewFormat.Height);
+				MediaPlayer->VideoTrackFormat.FrameRate = NewFormat.FrameRateNumerator / NewFormat.FrameRateDenominator;
+
+				MediaPlayer->SetupSampleChannels();
+
+				MediaPlayer->NotifyVideoFormatDetected();
 			}
 			else
 			{
@@ -490,10 +498,8 @@ namespace BlackmagicMediaPlayerHelpers
 
 FBlackmagicMediaPlayer::FBlackmagicMediaPlayer(IMediaEventSink& InEventSink)
 	: Super(InEventSink)
-	, EventCallback(nullptr)
 	, AudioSamplePool(MakeUnique<FBlackmagicMediaAudioSamplePool>())
 	, TextureSamplePool(MakeUnique<FBlackmagicMediaTextureSamplePool>())
-	, bVerifyFrameDropCount(false)
 	, SupportedSampleTypes(EMediaIOSampleType::None)
 {
 }
@@ -520,9 +526,6 @@ void FBlackmagicMediaPlayer::Close()
 	//Disable all our channels from the monitor
 	Samples->EnableTimedDataChannels(this, EMediaIOSampleType::None);
 
-	UnregisterSampleBuffers();
-	UnregisterTextures();
-
 	Super::Close();
 }
 
@@ -709,9 +712,18 @@ void FBlackmagicMediaPlayer::SetupSampleChannels()
 	Samples->InitializeAudioBuffer(AudioSettings);
 }
 
-void FBlackmagicMediaPlayer::AddVideoSample(const TSharedRef<FMediaIOCoreTextureSampleBase>& InSample)
+void FBlackmagicMediaPlayer::AddVideoSampleAfterGPUTransfer_RenderThread(const TSharedRef<FMediaIOCoreTextureSampleBase>& InSample)
 {
-	Samples->AddVideo(InSample);
+	checkSlow(IsInRenderingThread());
+
+	const TSharedRef<FBlackmagicMediaTextureSample> BlackMagicSample =
+		StaticCastSharedRef<FBlackmagicMediaTextureSample, FMediaIOCoreTextureSampleBase>(InSample);
+
+	// From now, the internal video buffer can be released
+	BlackMagicSample->ReleaseBlackmagicInternalBuffer();
+
+	// Let the base class do the rest
+	Super::AddVideoSampleAfterGPUTransfer_RenderThread(InSample);
 }
 
 #undef LOCTEXT_NAMESPACE
diff --git a/Engine/Plugins/Media/BlackmagicMedia/Source/BlackmagicMedia/Private/Player/BlackmagicMediaPlayer.h b/Engine/Plugins/Media/BlackmagicMedia/Source/BlackmagicMedia/Private/Player/BlackmagicMediaPlayer.h
--- a/Engine/Plugins/Media/BlackmagicMedia/Source/BlackmagicMedia/Private/Player/BlackmagicMediaPlayer.h
+++ b/Engine/Plugins/Media/BlackmagicMedia/Source/BlackmagicMedia/Private/Player/BlackmagicMediaPlayer.h
@@ -4,6 +4,8 @@
 
 #include "MediaIOCorePlayerBase.h"
 
+#include "BlackmagicMediaPrivate.h"
+
 #include "BlackmagicMediaSource.h"
 #include "HAL/CriticalSection.h"
 #include "MediaIOCoreAudioSampleBase.h"
@@ -23,6 +25,41 @@ namespace BlackmagicMediaPlayerHelpers
 
 class FBlackmagicMediaTextureSample : public FMediaIOCoreTextureSampleBase
 {
+public:
+	virtual const FMatrix& GetYUVToRGBMatrix() const override
+	{
+		return MediaShaders::YuvToRgbRec709Scaled;
+	}
+
+	//~ Begin IMediaPoolable interface
+public:
+	virtual void ShutdownPoolable() override
+	{
+		// Normally it should be explicitly released after GPU transfer. Just a bit of safety.
+		ReleaseBlackmagicInternalBuffer();
+
+		FMediaIOCoreTextureSampleBase::ShutdownPoolable();
+	}
+	//~ End IMediaPoolable interface
+
+public:
+
+	/** Returns buffer guard so it can be initialized outside (for performance reason) */
+	TSharedPtr<BlackmagicDesign::IInputEventCallback::FFrameBufferHolder>& GetBlackmagicInternalBufferLocker()
+	{
+		return VideoBufferGuard;
+	}
+
+	/** Unlock internal buffer referenced by this sample */
+	void ReleaseBlackmagicInternalBuffer()
+	{
+		VideoBufferGuard.Reset();
+	}
+
+private:
+
+	/** Buffer guard instance */
+	TSharedPtr<BlackmagicDesign::IInputEventCallback::FFrameBufferHolder> VideoBufferGuard;
 };
 
 class FBlackmagicMediaAudioSamplePool : public TMediaObjectPool<FMediaIOCoreAudioSampleBase> { };
@@ -95,8 +132,8 @@ protected:
 	{
 		return BlackmagicColorFormat == EBlackmagicMediaSourceColorFormat::YUV8 ? EMediaIOCoreColorFormat::YUV8 : EMediaIOCoreColorFormat::YUV10;
 	}
-	virtual void AddVideoSample(const TSharedRef<FMediaIOCoreTextureSampleBase>& InSample) override;
-	virtual TSharedPtr<FMediaIOCoreTextureSampleBase> AcquireSample_AnyThread() const override
+	virtual void AddVideoSampleAfterGPUTransfer_RenderThread(const TSharedRef<FMediaIOCoreTextureSampleBase>& InSample) override;
+	virtual TSharedPtr<FMediaIOCoreTextureSampleBase> AcquireTextureSample_AnyThread() const override
 	{
 		return TextureSamplePool->AcquireShared();
 	}
@@ -105,14 +142,14 @@ protected:
 private:
 
 	friend BlackmagicMediaPlayerHelpers::FBlackmagicMediaPlayerEventCallback;
-	BlackmagicMediaPlayerHelpers::FBlackmagicMediaPlayerEventCallback* EventCallback;
+	BlackmagicMediaPlayerHelpers::FBlackmagicMediaPlayerEventCallback* EventCallback = nullptr;
 
 	/** Audio, MetaData, Texture  sample object pool. */
 	TUniquePtr<FBlackmagicMediaAudioSamplePool> AudioSamplePool;
 	TUniquePtr<FBlackmagicMediaTextureSamplePool> TextureSamplePool;
 
 	/** Log warning about the amount of audio/video frame can't could not be cached . */
-	bool bVerifyFrameDropCount;
+	bool bVerifyFrameDropCount = false;
 
 	/** Max sample count our different buffer can hold. Taken from MediaSource */
 	int32 MaxNumAudioFrameBuffer = 0;
diff --git a/Engine/Plugins/Media/BlackmagicMedia/Source/BlackmagicMediaOutput/Public/BlackmagicMediaCapture.h b/Engine/Plugins/Media/BlackmagicMedia/Source/BlackmagicMediaOutput/Public/BlackmagicMediaCapture.h
--- a/Engine/Plugins/Media/BlackmagicMedia/Source/BlackmagicMediaOutput/Public/BlackmagicMediaCapture.h
+++ b/Engine/Plugins/Media/BlackmagicMedia/Source/BlackmagicMediaOutput/Public/BlackmagicMediaCapture.h
@@ -78,7 +78,7 @@ private:
 
 private:
 	friend BlackmagicMediaCaptureHelpers::FBlackmagicMediaCaptureEventCallback;
-	BlackmagicMediaCaptureHelpers::FBlackmagicMediaCaptureEventCallback* EventCallback;
+	BlackmagicMediaCaptureHelpers::FBlackmagicMediaCaptureEventCallback* EventCallback = nullptr;
 
 	/** Option from MediaOutput */
 	bool bWaitForSyncEvent;
diff --git a/Engine/Plugins/Media/MediaIOFramework/Source/GPUTextureTransfer/Private/TextureTransferBase.cpp b/Engine/Plugins/Media/MediaIOFramework/Source/GPUTextureTransfer/Private/TextureTransferBase.cpp
--- a/Engine/Plugins/Media/MediaIOFramework/Source/GPUTextureTransfer/Private/TextureTransferBase.cpp
+++ b/Engine/Plugins/Media/MediaIOFramework/Source/GPUTextureTransfer/Private/TextureTransferBase.cpp
@@ -119,7 +119,7 @@ namespace UE::GPUTextureTransfer::Private
 					DVP_CALL(dvpBegin());
 					constexpr uint64 NanosecondsToWait = 500000000; // 0.5 seconds
 					DVPStatus SyncStatus = dvpSyncObjClientWaitPartial(Info->GPUMemorySync->DVPSyncObject, Info->GPUMemorySync->AcquireValue, NanosecondsToWait);
-					if (SyncStatus == DVP_STATUS_INVALID_OPERATION)
+					if (SyncStatus != DVP_STATUS_OK)
 					{
 						bSuccess = false;
 						UE_LOG(LogGPUTextureTransfer, Error, TEXT("GPU Direct failed to sync."));
diff --git a/Engine/Plugins/Media/MediaIOFramework/Source/MediaIOCore/Private/CaptureCardMediaSource.cpp b/Engine/Plugins/Media/MediaIOFramework/Source/MediaIOCore/Private/CaptureCardMediaSource.cpp
--- a/Engine/Plugins/Media/MediaIOFramework/Source/MediaIOCore/Private/CaptureCardMediaSource.cpp
+++ b/Engine/Plugins/Media/MediaIOFramework/Source/MediaIOCore/Private/CaptureCardMediaSource.cpp
@@ -2,6 +2,12 @@
 
 #include "CaptureCardMediaSource.h"
 
+
+UCaptureCardMediaSource::UCaptureCardMediaSource()
+{
+	Deinterlacer = CreateDefaultSubobject<UBobDeinterlacer>("Deinterlacer");
+}
+
 int64 UCaptureCardMediaSource::GetMediaOption(const FName& Key, int64 DefaultValue) const
 {
 	if (Key == UE::CaptureCardMediaSource::InterlaceFieldOrder)
@@ -19,6 +25,11 @@ int64 UCaptureCardMediaSource::GetMediaOption(const FName& Key, int64 DefaultVal
 		return (int64) OverrideSourceEncoding;
 	}
 
+	if (Key == UE::CaptureCardMediaSource::EvaluationType)
+	{
+		return (int64) EvaluationType;
+	}
+
 	return UTimeSynchronizableMediaSource::GetMediaOption(Key, DefaultValue);
 }
 
@@ -37,6 +48,16 @@ FString UCaptureCardMediaSource::GetMediaOption(const FName& Key, const FString&
 
 bool UCaptureCardMediaSource::GetMediaOption(const FName& Key, bool bDefaultValue) const
 {
+	if (Key == UE::CaptureCardMediaSource::RenderJIT)
+	{
+		return bRenderJIT;
+	}
+
+	if (Key == UE::CaptureCardMediaSource::Framelock)
+	{
+		return bFramelock;
+	}
+
 	if (Key == UE::CaptureCardMediaSource::OverrideSourceColorSpace)
 	{
 		return bOverrideSourceColorSpace;
@@ -57,10 +78,61 @@ bool UCaptureCardMediaSource::HasMediaOption(const FName& Key) const
 		|| Key == UE::CaptureCardMediaSource::OverrideSourceEncoding
 		|| Key == UE::CaptureCardMediaSource::OverrideSourceColorSpace
 		|| Key == UE::CaptureCardMediaSource::SourceEncoding
-		|| Key == UE::CaptureCardMediaSource::SourceColorSpace)
+		|| Key == UE::CaptureCardMediaSource::SourceColorSpace
+		|| Key == UE::CaptureCardMediaSource::RenderJIT
+		|| Key == UE::CaptureCardMediaSource::Framelock
+		|| Key == UE::CaptureCardMediaSource::EvaluationType)
 		{
 			return true;
 		}
 
 	return UTimeSynchronizableMediaSource::HasMediaOption(Key);
 }
+
+#if WITH_EDITOR
+bool UCaptureCardMediaSource::CanEditChange(const FProperty* InProperty) const
+{
+	if (!Super::CanEditChange(InProperty))
+	{
+		return false;
+	}
+
+	if (InProperty->GetFName() == GET_MEMBER_NAME_CHECKED(UCaptureCardMediaSource, bFramelock))
+	{
+		return bRenderJIT && EvaluationType == EMediaIOSampleEvaluationType::Timecode && bUseTimeSynchronization;
+	}
+
+	return true;
+}
+
+void UCaptureCardMediaSource::PostEditChangeChainProperty(struct FPropertyChangedChainEvent& PropertyChangedEvent)
+{
+	Super::PostEditChangeChainProperty(PropertyChangedEvent);
+
+	if (PropertyChangedEvent.Property->GetFName() == GET_MEMBER_NAME_CHECKED(UCaptureCardMediaSource, bUseTimeSynchronization))
+	{
+		if (!bUseTimeSynchronization)
+		{
+			bFramelock = false;
+			EvaluationType = EMediaIOSampleEvaluationType::PlatformTime;
+		}
+	}
+
+	if (PropertyChangedEvent.Property->GetFName() == GET_MEMBER_NAME_CHECKED(UCaptureCardMediaSource, bRenderJIT))
+	{
+		if (!bRenderJIT)
+		{
+			bFramelock = false;
+			EvaluationType = (EvaluationType == EMediaIOSampleEvaluationType::Latest ? EMediaIOSampleEvaluationType::PlatformTime : EvaluationType);
+		}
+	}
+
+	if (PropertyChangedEvent.Property->GetFName() == GET_MEMBER_NAME_CHECKED(UCaptureCardMediaSource, EvaluationType))
+	{
+		if (EvaluationType != EMediaIOSampleEvaluationType::Timecode)
+		{
+			bFramelock = false;
+		}
+	}
+}
+#endif //WITH_EDITOR
diff --git a/Engine/Plugins/Media/MediaIOFramework/Source/MediaIOCore/Private/MediaIOCoreTextureSampleConverter.cpp b/Engine/Plugins/Media/MediaIOFramework/Source/MediaIOCore/Private/MediaIOCoreTextureSampleConverter.cpp
new file mode 100644
--- /dev/null
+++ b/Engine/Plugins/Media/MediaIOFramework/Source/MediaIOCore/Private/MediaIOCoreTextureSampleConverter.cpp
@@ -0,0 +1,43 @@
+// Copyright Epic Games, Inc. All Rights Reserved.
+
+#include "MediaIOCoreTextureSampleConverter.h"
+
+#include "MediaIOCorePlayerBase.h"
+#include "MediaIOCoreTextureSampleBase.h"
+
+
+void FMediaIOCoreTextureSampleConverter::Setup(const TSharedPtr<FMediaIOCoreTextureSampleBase>& InSample)
+{
+	// This proxy sample will be exposed to the media framework
+	JITRProxySample = InSample;
+}
+
+bool FMediaIOCoreTextureSampleConverter::Convert(FTexture2DRHIRef& InDestinationTexture, const FConversionHints& Hints)
+{
+	TRACE_CPUPROFILER_EVENT_SCOPE(JITR);
+
+	// Pin the sample instance
+	TSharedPtr<FMediaIOCoreTextureSampleBase> SamplePtr = JITRProxySample.Pin();
+	if (SamplePtr.IsValid() == false)
+	{
+		return false;
+	}
+
+	// Get owning player to run JITR path
+	TSharedPtr<FMediaIOCorePlayerBase> PlayerPtr = SamplePtr->GetPlayer();
+	if (PlayerPtr.IsValid() == false)
+	{
+		return false;
+	}
+
+	// Let the player pick a proper sample and render it into this proxy's texture (JITR)
+	return PlayerPtr->JustInTimeSampleRender_RenderThread(SamplePtr);
+}
+
+uint32 FMediaIOCoreTextureSampleConverter::GetConverterInfoFlags() const
+{
+	// JITR requires PreprocessOnly conversion flag as the main idea is to replace
+	// the dummy data of the JITR proxy sample with a real data of a sample that
+	// will be chosen down the pipe. However, it may return other values when inherited.
+	return IMediaTextureSampleConverter::ConverterInfoFlags_PreprocessOnly;
+}
diff --git a/Engine/Plugins/Media/MediaIOFramework/Source/MediaIOCore/Private/Player/MediaIOCorePlayerBase.cpp b/Engine/Plugins/Media/MediaIOFramework/Source/MediaIOCore/Private/Player/MediaIOCorePlayerBase.cpp
--- a/Engine/Plugins/Media/MediaIOFramework/Source/MediaIOCore/Private/Player/MediaIOCorePlayerBase.cpp
+++ b/Engine/Plugins/Media/MediaIOFramework/Source/MediaIOCore/Private/Player/MediaIOCorePlayerBase.cpp
@@ -3,6 +3,7 @@
 #include "MediaIOCorePlayerBase.h"
 
 #include "CaptureCardMediaSource.h"
+#include "Async/Async.h"
 #include "Engine/Engine.h"
 #include "Engine/TimecodeProvider.h"
 #include "GPUTextureTransferModule.h"
@@ -12,9 +13,15 @@
 #include "IMediaModule.h"
 #include "IMediaPlayerFactory.h"
 #include "ITimeManagementModule.h"
+#include "MediaIOCoreAudioSampleBase.h"
+#include "MediaIOCoreBinarySampleBase.h"
+#include "MediaIOCoreCaptionSampleBase.h"
 #include "MediaIOCoreDeinterlacer.h"
 #include "MediaIOCoreModule.h"
 #include "MediaIOCoreSamples.h"
+#include "MediaIOCoreSubtitleSampleBase.h"
+#include "MediaIOCoreTextureSampleBase.h"
+#include "MediaIOCoreTextureSampleConverter.h"
 #include "Misc/App.h"
 #include "Misc/ScopeLock.h"
 #include "RenderingThread.h"
@@ -48,6 +55,13 @@ namespace MediaIOCorePlayerDetail
 		ECVF_Default
 	);
 
+	// [JITR] Detailed insights
+	static TAutoConsoleVariable<bool> CVarMediaIOJITRInsights(
+		TEXT("MediaIO.JITR.Insights"),
+		false,
+		TEXT("Generate detailed Just-in-Time Media Rendering insights"),
+		ECVF_Default);
+
 	static FAutoConsoleCommand CCommandShowTimecode(
 		TEXT("MediaIO.ShowInputTimecode"),
 		TEXT("All media player will log the frame timecode when a new frame is captured."),
@@ -64,9 +78,11 @@ namespace MediaIOCorePlayerDetail
 		TEXT("MediaIO.EnableGPUDirectInput"), 0,
 		TEXT("Whether to enable GPU direct for faster video frame copies. (Experimental)"),
 		ECVF_RenderThreadSafe);
-
 }
 
+DECLARE_GPU_STAT_NAMED(STAT_MediaIOPlayer_JITR_TransferTexture, TEXT("JITR_TransferTexture"));
+
+
 /* FMediaIOCoreMediaOption structors
  *****************************************************************************/
 
@@ -80,27 +96,16 @@ const FName FMediaIOCoreMediaOption::VideoModeName("VideoModeName");
  *****************************************************************************/
 
 FMediaIOCorePlayerBase::FMediaIOCorePlayerBase(IMediaEventSink& InEventSink)
-	: CurrentState(EMediaState::Closed)
-	, CurrentTime(FTimespan::Zero())
-	, EventSink(InEventSink)
-	, VideoFrameRate(30, 1)
-	, bUseTimeSynchronization(false)
-	, bWarnedIncompatibleFrameRate(false)
-	, bAutoDetect(false)
-	, FrameDelay(0)
-	, TimeDelay(0.0)
-	, PreviousFrameTimespan(FTimespan::Zero())
-{
-	Samples = new FMediaIOCoreSamples();
-
+	: EventSink(InEventSink)
+	, Samples(MakeUnique<FMediaIOCoreSamples>())
+	, JITRSamples(MakeUnique<FJITRMediaTextureSamples>())
+{
 	FGPUTextureTransferModule& Module = FGPUTextureTransferModule::Get();
 	GPUTextureTransfer = Module.GetTextureTransfer();
 }
 
 FMediaIOCorePlayerBase::~FMediaIOCorePlayerBase()
 {
-	delete Samples;
-	Samples = nullptr;
 }
 
 /* IMediaPlayer interface
@@ -118,6 +123,11 @@ void FMediaIOCorePlayerBase::Close()
 		AudioTrackFormat.SampleRate = 0;
 
 		Samples->FlushSamples();
+		JITRSamples->FlushSamples();
+
+		UnregisterSampleBuffers();
+		UnregisterTextures();
+
 		EventSink.ReceiveMediaEvent(EMediaEvent::TracksChanged);
 		EventSink.ReceiveMediaEvent(EMediaEvent::MediaClosed);
 	}
@@ -164,7 +174,16 @@ IMediaControls& FMediaIOCorePlayerBase::GetControls()
 
 IMediaSamples& FMediaIOCorePlayerBase::GetSamples()
 {
-	return *Samples;
+	const bool bIsJITREnabled = IsJustInTimeRenderingEnabled();
+
+	if (bIsJITREnabled)
+	{
+		return *JITRSamples;
+	}
+	else
+	{
+		return *Samples;
+	}
 }
 
 const FMediaIOCoreSamples& FMediaIOCorePlayerBase::GetSamples() const
@@ -201,9 +220,11 @@ bool FMediaIOCorePlayerBase::Open(const FString& Url, const IMediaOptions* Optio
 	OpenUrl = Url;
 	const bool bReadMediaOptions = ReadMediaOptions(Options);
 
-	if (MediaIOCorePlayerDetail::CVarEnableGPUDirectInput.GetValueOnAnyThread())
+	// When AutoDetect is true, we can't create textures now as we don't know the video format.
+	// In this case we need to wait until auto-detection is finished, and video format is known.
+	if (!bAutoDetect && CanUseGPUTextureTransfer())
 	{
-		CreateAndRegisterTextures(Options);
+		CreateAndRegisterTextures();
 	}
 	
 	const EMediaIOInterlaceFieldOrder InterlaceFieldOrder = static_cast<EMediaIOInterlaceFieldOrder>(Options->GetMediaOption(UE::CaptureCardMediaSource::InterlaceFieldOrder, (int64)EMediaIOInterlaceFieldOrder::TopFieldFirst));
@@ -215,11 +236,11 @@ bool FMediaIOCorePlayerBase::Open(const FString& Url, const IMediaOptions* Optio
 	TSharedPtr<UE::MediaIOCore::FDeinterlacer> DeinterlacerInstance;
 	if (DeinterlacerInstancer && DeinterlacerInstancer->IsA<UVideoDeinterlacer>())
 	{
-		Deinterlacer = Cast<UVideoDeinterlacer>(DeinterlacerInstancer)->Instantiate(UE::MediaIOCore::FDeinterlacer::FOnAcquireSample_AnyThread::CreateSP(this, &FMediaIOCorePlayerBase::AcquireSample_AnyThread), InterlaceFieldOrder);
+		Deinterlacer = Cast<UVideoDeinterlacer>(DeinterlacerInstancer)->Instantiate(UE::MediaIOCore::FDeinterlacer::FOnAcquireSample_AnyThread::CreateSP(this, &FMediaIOCorePlayerBase::AcquireTextureSample_AnyThread), InterlaceFieldOrder);
 	}
 	else
 	{
-		Deinterlacer = MakeShared<UE::MediaIOCore::FDeinterlacer>(UE::MediaIOCore::FDeinterlacer::FOnAcquireSample_AnyThread::CreateSP(this, &FMediaIOCorePlayerBase::AcquireSample_AnyThread), InterlaceFieldOrder);
+		Deinterlacer = MakeShared<UE::MediaIOCore::FDeinterlacer>(UE::MediaIOCore::FDeinterlacer::FOnAcquireSample_AnyThread::CreateSP(this, &FMediaIOCorePlayerBase::AcquireTextureSample_AnyThread), InterlaceFieldOrder);
 	}
 
 	return bReadMediaOptions;
@@ -232,7 +253,7 @@ bool FMediaIOCorePlayerBase::Open(const TSharedRef<FArchive, ESPMode::ThreadSafe
 
 void FMediaIOCorePlayerBase::TickTimeManagement()
 {
-	if (bUseTimeSynchronization)
+	if (EvaluationType == EMediaIOSampleEvaluationType::Timecode)
 	{
 		TOptional<FQualifiedFrameTime> CurrentFrameTime = FApp::GetCurrentFrameTime();
 		if (CurrentFrameTime.IsSet())
@@ -269,6 +290,17 @@ void FMediaIOCorePlayerBase::TickTimeManagement()
 
 void FMediaIOCorePlayerBase::TickFetch(FTimespan DeltaTime, FTimespan Timecode)
 {
+	// Running JITR?
+	if (CurrentState == EMediaState::Playing && IsJustInTimeRenderingEnabled())
+	{
+		// Nothing to do if no samples received yet
+		if (Samples->NumVideoSamples() > 0)
+		{
+			// Create new JITR proxy sample
+			JITRSamples->ProxySample = AcquireJITRProxySampleInitialized();
+		}
+	}
+
 	Samples->CacheSamplesState(GetTime());
 }
 
@@ -491,11 +523,17 @@ bool FMediaIOCorePlayerBase::SetTrackFormat(EMediaTrackType TrackType, int32 Tra
 
 bool FMediaIOCorePlayerBase::ReadMediaOptions(const IMediaOptions* Options)
 {
-	bUseTimeSynchronization = Options->GetMediaOption(TimeSynchronizableMedia::UseTimeSynchronizatioOption, false);
+	// Break dependency on the bUseTimeSynchronization option, and use EvaluationType instead
+	const bool bUseTimeSynchronization = Options->GetMediaOption(TimeSynchronizableMedia::UseTimeSynchronizatioOption, false);
+
 	FrameDelay = Options->GetMediaOption(TimeSynchronizableMedia::FrameDelay, (int64)0);
 	TimeDelay = Options->GetMediaOption(TimeSynchronizableMedia::TimeDelay, 0.0);
 	bAutoDetect = Options->GetMediaOption(TimeSynchronizableMedia::AutoDetect, false);
 
+	EvaluationType = (EMediaIOSampleEvaluationType)(Options->GetMediaOption(UE::CaptureCardMediaSource::EvaluationType, (int64)EMediaIOSampleEvaluationType::PlatformTime));
+	bJustInTimeRender = Options->GetMediaOption(UE::CaptureCardMediaSource::RenderJIT, true);
+	bFramelock = Options->GetMediaOption(UE::CaptureCardMediaSource::Framelock, false);
+
 	{
 		int32 Numerator = Options->GetMediaOption(FMediaIOCoreMediaOption::FrameRateNumerator, (int64)30);
 		int32 Denominator = Options->GetMediaOption(FMediaIOCoreMediaOption::FrameRateDenominator, (int64)1);
@@ -510,8 +548,20 @@ bool FMediaIOCorePlayerBase::ReadMediaOptions(const IMediaOptions* Options)
 		VideoTrackFormat.TypeName = Options->GetMediaOption(FMediaIOCoreMediaOption::VideoModeName, FString(TEXT("1080p30")));
 	}
 
+	// Make sure we haven't got invalid data from media source. When time sync is off,
+	// timecode based evaluation and framelock are unavailable.
+	if (!bUseTimeSynchronization)
+	{
+		if (EvaluationType == EMediaIOSampleEvaluationType::Timecode)
+		{
+			EvaluationType = EMediaIOSampleEvaluationType::PlatformTime;
+		}
+		
+		bFramelock = false;
+	}
+
 	//Setup base sampling settings
-	BaseSettings.EvaluationType = bUseTimeSynchronization ? ETimedDataInputEvaluationType::Timecode : ETimedDataInputEvaluationType::PlatformTime;
+	BaseSettings.EvaluationType = GetEvaluationType(); // bJustInTimeRender and JITREvaluationType should have been set before calling
 	BaseSettings.FrameRate = VideoFrameRate;
 	BaseSettings.PlayerTimeOffset = MediaIOCorePlayerDetail::HighPerformanceClockDelta;
 	BaseSettings.AbsoluteMaxBufferSize = MediaIOCorePlayerDetail::CVarMediaIOMaxBufferSize.GetValueOnGameThread();
@@ -519,6 +569,15 @@ bool FMediaIOCorePlayerBase::ReadMediaOptions(const IMediaOptions* Options)
 	return true;
 }
 
+void FMediaIOCorePlayerBase::NotifyVideoFormatDetected()
+{
+	// Now that we know the video format, we can initialize the textures
+	if (bAutoDetect && CanUseGPUTextureTransfer())
+	{
+		CreateAndRegisterTextures();
+	}
+}
+
 double FMediaIOCorePlayerBase::GetApplicationSeconds()
 {
 	return MediaIOCorePlayerDetail::HighPerformanceClockDelta + FApp::GetCurrentTime();
@@ -538,11 +597,34 @@ bool FMediaIOCorePlayerBase::IsTimecodeLogEnabled()
 #endif
 }
 
-bool FMediaIOCorePlayerBase::CanUseGPUTextureTransfer()
+bool FMediaIOCorePlayerBase::CanUseGPUTextureTransfer() const
 {
 	return GPUTextureTransfer && MediaIOCorePlayerDetail::CVarEnableGPUDirectInput.GetValueOnAnyThread();
 }
 
+bool FMediaIOCorePlayerBase::HasTextureAvailableForGPUTransfer() const
+{
+	FScopeLock ScopeLock(&TexturesCriticalSection);
+	return !Textures.IsEmpty();
+}
+
+bool FMediaIOCorePlayerBase::IsJustInTimeRenderingEnabled() const
+{
+	return bJustInTimeRender;
+}
+
+void FMediaIOCorePlayerBase::AddVideoSampleAfterGPUTransfer_RenderThread(const TSharedRef<FMediaIOCoreTextureSampleBase>& Sample)
+{
+	checkSlow(IsInRenderingThread());
+
+	LogBookmark(TEXT("MIO Tex Copy"), Sample);
+
+	if (!IsJustInTimeRenderingEnabled())
+	{
+		Samples->AddVideo(Sample);
+	}
+}
+
 void FMediaIOCorePlayerBase::OnSampleDestroyed(TRefCountPtr<FRHITexture> InTexture)
 {
 	FScopeLock ScopeLock(&TexturesCriticalSection);
@@ -615,10 +697,24 @@ void FMediaIOCorePlayerBase::UnregisterSampleBuffers()
 	RegisteredBuffers.Reset();
 }
 
-void FMediaIOCorePlayerBase::CreateAndRegisterTextures(const IMediaOptions* Options)
+void FMediaIOCorePlayerBase::CreateAndRegisterTextures()
 {
+	// Call this on the game thread so we can stall until textures are created
+	if (!IsInGameThread())
+	{
+		AsyncTask(ENamedThreads::GameThread, [this]()
+			{
+				CreateAndRegisterTextures();
+			});
+
+		return;
+	}
+
+
+	UnregisterTextures();
+
 	ENQUEUE_RENDER_COMMAND(MediaIOCorePlayerCreateTextures)(
-		[this, Options](FRHICommandListImmediate& CommandList)
+		[this](FRHICommandListImmediate& CommandList)
 		{
 			if (!GPUTextureTransfer)
 			{
@@ -626,7 +722,16 @@ void FMediaIOCorePlayerBase::CreateAndRegisterTextures(const IMediaOptions* Opti
 				return;
 			}
 
-			for (uint8 Index = 0; Index < GetNumVideoFrameBuffers(); Index++)
+			// We need only one texture for JIT rendering. However the original non-JITR pipeline
+			// requires the number of textures must not be less than the number of buffers
+			const bool   bIsJustInTimeRenderingEnabled = IsJustInTimeRenderingEnabled();
+			const uint32 VideoFrameBuffersAmount = GetNumVideoFrameBuffers();
+			const uint32 TexturesAmount = bIsJustInTimeRenderingEnabled ? 1 : VideoFrameBuffersAmount;
+
+			RegisteredTextures.Reserve(TexturesAmount);
+			Textures.Reserve(TexturesAmount);
+
+			for (uint8 Index = 0; Index < TexturesAmount; Index++)
 			{
 				TRefCountPtr<FRHITexture> RHITexture;
 				const FString TextureName = FString::Printf(TEXT("FMediaIOCorePlayerTexture %d"), Index);
@@ -721,18 +826,32 @@ void FMediaIOCorePlayerBase::UnregisterTextures()
 
 void FMediaIOCorePlayerBase::PreGPUTransfer(const TSharedPtr<FMediaIOCoreTextureSampleBase>& InSample)
 {
+	TRefCountPtr<FRHITexture> Texture;
+
+	// Get a texture from the pool
+	{
+		FScopeLock Lock(&TexturesCriticalSection);
+
+		const bool bHasTexture = !Textures.IsEmpty();
+		ensureMsgf(bHasTexture, TEXT("No texture available while doing a gpu texture transfer."));
+
+		if (bHasTexture)
+		{
+			Texture = Textures.Pop();
+		}
+	}
+
 	// Register the buffer and the texture
-	TRefCountPtr<FRHITexture> Texture = Textures.Pop();
 	if (Texture && RegisteredTextures.Contains(Texture))
 	{
 		InSample->SetTexture(Texture);
 		InSample->SetDestructionCallback([MediaPlayerWeakPtr = TWeakPtr<FMediaIOCorePlayerBase>(AsShared())](TRefCountPtr<FRHITexture> InTexture)
-		{
-			if (TSharedPtr<FMediaIOCorePlayerBase> MediaPlayerPtr = MediaPlayerWeakPtr.Pin())
 			{
-				MediaPlayerPtr->OnSampleDestroyed(InTexture);
-			}
-		});
+				if (TSharedPtr<FMediaIOCorePlayerBase> MediaPlayerPtr = MediaPlayerWeakPtr.Pin())
+				{
+					MediaPlayerPtr->OnSampleDestroyed(InTexture);
+				}
+			});
 
 		RegisterSampleBuffer(InSample);
 	}
@@ -742,6 +861,37 @@ void FMediaIOCorePlayerBase::PreGPUTransfer(const TSharedPtr<FMediaIOCoreTexture
 	}
 }
 
+void FMediaIOCorePlayerBase::PreGPUTransferJITR(const TSharedPtr<FMediaIOCoreTextureSampleBase>& InSample, const TSharedPtr<FMediaIOCoreTextureSampleBase>& InJITRProxySample)
+{
+	TRefCountPtr<FRHITexture> Texture;
+
+	// Get a texture from the pool. JITR path uses a single texture for GPU transfer.
+	{
+		FScopeLock Lock(&TexturesCriticalSection);
+
+		const bool bHasTexture = !Textures.IsEmpty();
+
+		// Normally this should never yell as we should always have a single texture
+		ensureMsgf(bHasTexture, TEXT("No texture available while doing a gpu texture transfer."));
+
+		Texture = Textures[0];
+	}
+
+	// Register the buffer and the texture
+	if (Texture && RegisteredTextures.Contains(Texture))
+	{
+		// Set data that is necessary to perform fast texture copy right to the proxy sample
+		InJITRProxySample->SetTexture(Texture);
+		InJITRProxySample->SetBuffer(InSample->GetMutableBuffer());
+
+		RegisterSampleBuffer(InJITRProxySample);
+	}
+	else
+	{
+		UE_LOG(LogMediaIOCore, Display, TEXT("Unregistered texture %u encountered while doing a gpu texture transfer."), Texture ? reinterpret_cast<uintptr_t>(Texture->GetNativeResource()) : 0);
+	}
+}
+
 void FMediaIOCorePlayerBase::ExecuteGPUTransfer(const TSharedPtr<FMediaIOCoreTextureSampleBase>& InSample)
 {
 	ENQUEUE_RENDER_COMMAND(StartGPUTextureTransferCopy)(
@@ -756,7 +906,7 @@ void FMediaIOCorePlayerBase::ExecuteGPUTransfer(const TSharedPtr<FMediaIOCoreTex
 				UE_LOG(LogMediaIOCore, Verbose, TEXT("Starting Transfer with buffer %u"), reinterpret_cast<uintptr_t>(Buffer));
 				MediaPlayerPtr->GPUTextureTransfer->TransferTexture(Buffer, InSample->GetTexture(), UE::GPUTextureTransfer::ETransferDirection::CPU_TO_GPU);
 				MediaPlayerPtr->GPUTextureTransfer->BeginSync((InSample)->GetMutableBuffer(), UE::GPUTextureTransfer::ETransferDirection::CPU_TO_GPU);
-				MediaPlayerPtr->AddVideoSample(InSample.ToSharedRef());
+				MediaPlayerPtr->AddVideoSampleAfterGPUTransfer_RenderThread(InSample.ToSharedRef());
 				MediaPlayerPtr->GPUTextureTransfer->EndSync(InSample->GetMutableBuffer());
 			}
 		}
@@ -793,64 +943,100 @@ FText FMediaIOCorePlayerBase::GetDisplayName() const
 
 	return FText::Format(LOCTEXT("PlayerDisplayName", "{0} - {1}"), FText::FromName(PlayerFactory->GetPlayerName()), FText::FromString(GetUrl()));
 }
+
 ETimedDataInputEvaluationType FMediaIOCorePlayerBase::GetEvaluationType() const
 {
-	if (bUseTimeSynchronization)
-	{
-		return ETimedDataInputEvaluationType::Timecode;
-	}
-	else
+	switch (EvaluationType)
 	{
+	case EMediaIOSampleEvaluationType::Latest:
+		return ETimedDataInputEvaluationType::None;
+
+	case EMediaIOSampleEvaluationType::PlatformTime:
 		return ETimedDataInputEvaluationType::PlatformTime;
+
+	case EMediaIOSampleEvaluationType::Timecode:
+		return ETimedDataInputEvaluationType::Timecode;
+
+	default:
+		checkNoEntry();
+		return ETimedDataInputEvaluationType::None;
 	}
 }
 
 void FMediaIOCorePlayerBase::SetEvaluationType(ETimedDataInputEvaluationType Evaluation)
 {
-	bool bNewEvaluation = false;
+	bool bNewEvaluationModeSet = false;
+
 	switch (Evaluation)
 	{
-	case ETimedDataInputEvaluationType::Timecode:
-		bNewEvaluation = true;
+	case ETimedDataInputEvaluationType::None:
+		bNewEvaluationModeSet = (EvaluationType != EMediaIOSampleEvaluationType::Latest);
+		EvaluationType = EMediaIOSampleEvaluationType::Latest;
 		break;
+
 	case ETimedDataInputEvaluationType::PlatformTime:
-	case ETimedDataInputEvaluationType::None:
+		bNewEvaluationModeSet = (EvaluationType != EMediaIOSampleEvaluationType::PlatformTime);
+		EvaluationType = EMediaIOSampleEvaluationType::PlatformTime;
+		break;
+
+	case ETimedDataInputEvaluationType::Timecode:
+		bNewEvaluationModeSet = (EvaluationType != EMediaIOSampleEvaluationType::Timecode);
+		EvaluationType = EMediaIOSampleEvaluationType::Timecode;
+		break;
+
 	default:
-		bNewEvaluation = false;
+		checkNoEntry();
 		break;
 	}
 
-	if (bNewEvaluation != bUseTimeSynchronization)
+	if (bNewEvaluationModeSet)
 	{
-		bUseTimeSynchronization = bNewEvaluation;
-		BaseSettings.EvaluationType = bUseTimeSynchronization ? ETimedDataInputEvaluationType::Timecode : ETimedDataInputEvaluationType::PlatformTime;
+		BaseSettings.EvaluationType = GetEvaluationType();
 		SetupSampleChannels();
 	}
 }
 
 double FMediaIOCorePlayerBase::GetEvaluationOffsetInSeconds() const
 {
-	if (bUseTimeSynchronization)
-	{
-		return ITimedDataInput::ConvertFrameOffsetInSecondOffset(FrameDelay, VideoFrameRate);
-	}
-	else
+	switch (EvaluationType)
 	{
+	case EMediaIOSampleEvaluationType::Latest:
+		return 0;
+
+	case EMediaIOSampleEvaluationType::PlatformTime:
 		return TimeDelay;
+
+	case EMediaIOSampleEvaluationType::Timecode:
+		return ITimedDataInput::ConvertFrameOffsetInSecondOffset(FrameDelay, VideoFrameRate);
+
+	default:
+		checkNoEntry();
+		return 0;
 	}
 }
 
 void FMediaIOCorePlayerBase::SetEvaluationOffsetInSeconds(double Offset)
 {
-	if (bUseTimeSynchronization)
+	switch (EvaluationType)
+	{
+	case EMediaIOSampleEvaluationType::Latest:
+		break;
+
+	case EMediaIOSampleEvaluationType::PlatformTime:
+		TimeDelay = Offset;
+		break;
+
+	case EMediaIOSampleEvaluationType::Timecode:
 	{
 		//Media doesn't support subframes playback (interpolation between frames) so offsets are always in full frame
 		const double FrameOffset = ITimedDataInput::ConvertSecondOffsetInFrameOffset(Offset, VideoFrameRate);
 		FrameDelay = FMath::CeilToInt(FrameOffset);
+		break;
 	}
-	else
-	{
-		TimeDelay = Offset;
+
+	default:
+		checkNoEntry();
+		break;
 	}
 }
 
@@ -880,4 +1066,305 @@ bool FMediaIOCorePlayerBase::SupportsSubFrames() const
 	return false;
 }
 
+void FMediaIOCorePlayerBase::AddAudioSample(const TSharedRef<FMediaIOCoreAudioSampleBase>& Sample)
+{
+	Samples->AddAudio(Sample);
+}
+
+void FMediaIOCorePlayerBase::AddCaptionSample(const TSharedRef<FMediaIOCoreCaptionSampleBase>& Sample)
+{
+	Samples->AddCaption(Sample);
+}
+
+void FMediaIOCorePlayerBase::AddMetadataSample(const TSharedRef<FMediaIOCoreBinarySampleBase>& Sample)
+{
+	Samples->AddMetadata(Sample);
+}
+
+void FMediaIOCorePlayerBase::AddSubtitleSample(const TSharedRef<FMediaIOCoreSubtitleSampleBase>& Sample)
+{
+	Samples->AddSubtitle(Sample);
+}
+
+void FMediaIOCorePlayerBase::AddVideoSample(const TSharedRef<FMediaIOCoreTextureSampleBase>& Sample)
+{
+	LogBookmark(TEXT("MIO in"), Sample);
+
+	// When JITR is used, we postpone texture transferring until this sample is chosen to render. It's also
+	// possible that this sample will never be rendered, therefore we don't waste resources for transferring.
+	if (IsJustInTimeRenderingEnabled())
+	{
+		Samples->AddVideo(Sample);
+	}
+	// Otherwise we run original non-JITR pipeline which implies immediate texture transfer if GPUDirect path is available.
+	else
+	{
+		const bool bCanUseGPUTextureTransfer = CanUseGPUTextureTransfer();
+		const bool bIsAwaitingForGPUTransfer = Sample->IsAwaitingForGPUTransfer();
+
+		// So, when GPUDirect is availalbe we transfer data on the rendering thread. The callback AddVideoSampleAfterGPUTransfer_RenderThread
+		// will put the sample to the pool when finished.
+		if (bCanUseGPUTextureTransfer && bIsAwaitingForGPUTransfer)
+		{
+			PreGPUTransfer(Sample);
+			ExecuteGPUTransfer(Sample);
+		}
+		// Just push the sample to the pool if fast GPU transfer isn't available
+		else
+		{
+			Samples->AddVideo(Sample);
+		}
+	}
+}
+
+bool FMediaIOCorePlayerBase::JustInTimeSampleRender_RenderThread(TSharedPtr<FMediaIOCoreTextureSampleBase>& JITRProxySample)
+{
+	checkSlow(IsInRenderingThread());
+
+	// Make sure the sample is valid
+	if (!JITRProxySample.IsValid())
+	{
+		UE_LOG(LogMediaIOCore, Warning, TEXT("Invalid JITR sample."));
+		return false;
+	}
+
+	// JIT rendering is performed once per frame
+	if (LastEngineRTFrameThatUpdatedJustInTime == GFrameCounterRenderThread)
+	{
+		return false;
+	}
+
+	// Pick a sample to render
+	TSharedPtr<FMediaIOCoreTextureSampleBase> SourceSample = PickSampleToRender_RenderThread(JITRProxySample);
+	if (!SourceSample)
+	{
+		UE_LOG(LogMediaIOCore, Warning, TEXT("JustInTimeSampleRender couldn't find a sample to render"));
+		return false;
+	}
+
+	LogBookmark(TEXT("JITR out"), SourceSample.ToSharedRef());
+
+	// Update proxy sample configuration
+	JITRProxySample->CopyConfiguration(SourceSample);
+
+	// Now we know which sample to use, transfer its texture
+	TransferTexture_RenderThread(SourceSample, JITRProxySample);
+
+	// Update frame number so we won't render this sample again
+	LastEngineRTFrameThatUpdatedJustInTime = GFrameCounterRenderThread;
+
+	return true;
+}
+
+TSharedPtr<FMediaIOCoreTextureSampleConverter> FMediaIOCorePlayerBase::CreateTextureSampleConverter() const
+{
+	return MakeShared<FMediaIOCoreTextureSampleConverter>();
+}
+
+TSharedPtr<FMediaIOCoreTextureSampleBase> FMediaIOCorePlayerBase::AcquireJITRProxySampleInitialized()
+{
+	// Create a new media sample acting as a dummy container to be picked by MFW which we will fill during the late update
+	FMediaIOCoreSampleJITRConfigurationArgs Args;
+	Args.Width     = 1; // Dummy value, will be replaced on rendering
+	Args.Height    = 1; // Dummy value, will be replaced on rendering
+	Args.Player    = AsShared().ToSharedPtr();
+	Args.Time      = FTimespan::FromSeconds(GetPlatformSeconds());
+	Args.Timecode  = FApp::GetTimecode();
+	Args.EvaluationOffsetInSeconds = GetEvaluationOffsetInSeconds();
+	Args.Converter = CreateTextureSampleConverter();
+
+	check(Args.Converter.IsValid());
+
+	// This sample is going to be used for rendering
+	TSharedPtr<FMediaIOCoreTextureSampleBase> NewSample = AcquireTextureSample_AnyThread();
+	if (!NewSample.IsValid())
+	{
+		return nullptr;
+	}
+
+	// Setup sample to the converter
+	Args.Converter->Setup(NewSample);
+
+	// Initialize the JITR sample
+	const bool bInitialized = NewSample->InitializeJITR(Args);
+	if (!bInitialized)
+	{
+		return nullptr;
+	}
+
+	return NewSample;
+}
+
+TSharedPtr<FMediaIOCoreTextureSampleBase> FMediaIOCorePlayerBase::PickSampleToRender_RenderThread(const TSharedPtr<FMediaIOCoreTextureSampleBase>& JITRProxySample)
+{
+	switch (EvaluationType)
+	{
+	// Pick the latest available sample
+	case EMediaIOSampleEvaluationType::Latest:
+		return PickSampleToRenderForLatest_RenderThread(JITRProxySample);
+
+	// Pick a sample based on the platform time
+	case EMediaIOSampleEvaluationType::PlatformTime:
+		return PickSampleToRenderForTimeSynchronized_RenderThread(JITRProxySample);
+
+	// Pick a sample based on the engine's timecode
+	case EMediaIOSampleEvaluationType::Timecode:
+		return (bFramelock ?
+			PickSampleToRenderFramelocked_RenderThread(JITRProxySample) :
+			PickSampleToRenderForTimeSynchronized_RenderThread(JITRProxySample));
+	default:
+		checkNoEntry();
+		return PickSampleToRenderForLatest_RenderThread(JITRProxySample);
+	}
+}
+
+TSharedPtr<FMediaIOCoreTextureSampleBase> FMediaIOCorePlayerBase::PickSampleToRenderForLatest_RenderThread(const TSharedPtr<FMediaIOCoreTextureSampleBase>& JITRProxySample)
+{
+	const TArray<TSharedPtr<IMediaTextureSample>> TextureSamples = Samples->GetVideoSamples();
+
+	// Just return the most recent sample available in the queue
+	return TextureSamples.Num() > 0 ?
+		StaticCastSharedPtr<FMediaIOCoreTextureSampleBase, IMediaTextureSample, ESPMode::ThreadSafe>(TextureSamples[0]) :
+		nullptr;
+}
+
+TSharedPtr<FMediaIOCoreTextureSampleBase> FMediaIOCorePlayerBase::PickSampleToRenderForTimeSynchronized_RenderThread(const TSharedPtr<FMediaIOCoreTextureSampleBase>& JITRProxySample)
+{
+	// Reference time based on evaluation type
+	FTimespan TargetSampleTimespan;
+
+	// Get base uncorrected reference point
+	const FTimecode InvalidTimecode;
+	const FTimecode RequestedTimecode = JITRProxySample->GetTimecode().Get(InvalidTimecode);
+	if (EvaluationType == EMediaIOSampleEvaluationType::Timecode && RequestedTimecode != FTimecode())
+	{
+		// We'll use timecode data to find a proper sample
+		TargetSampleTimespan = RequestedTimecode.ToTimespan(VideoFrameRate);
+	}
+	else
+	{
+		// We'll use platform time to find a proper sample
+		TargetSampleTimespan = JITRProxySample->GetTime().Time;
+	}
+
+	// Apply time correction to the target time
+	const double RequestedOffsetInSeconds = JITRProxySample->GetEvaluationOffsetInSeconds();
+	const FTimespan RequestedOffsetTimespan = FTimespan::FromSeconds(RequestedOffsetInSeconds);
+	const FTimespan TargetTimespanCorrected = TargetSampleTimespan + RequestedOffsetTimespan;
+
+	// Go over the sample pool and find a sample closest to the target time
+	int32 ClosestIndex = -1;
+	int64 SmallestInterval(TNumericLimits<int64>::Max());
+
+	// Get all available video samples
+	const TArray<TSharedPtr<IMediaTextureSample>> TextureSamples = Samples->GetVideoSamples();
+
+	for (int32 Index = 0; Index < TextureSamples.Num(); ++Index)
+	{
+		// When EvaluationType == ETimedDataInputEvaluationType::Timecode, the time represents the sample's timecode.
+		// When EvaluationType == ETimedDataInputEvaluationType::PlatformTime, the time is based on the platform time.
+		const FTimespan TestTimespan = TextureSamples[Index]->GetTime().Time;
+
+		// Either closest positive or closest negative
+		const int64 TestInterval = FMath::Abs((TestTimespan - TargetTimespanCorrected).GetTicks());
+
+		// '<=' instead of '<' is used here intentionally. Turns out we might have
+		// some samples with the same timecode. To avoid early termination of the search '<=' is used.
+		if (TestInterval <= SmallestInterval)
+		{
+			ClosestIndex = Index;
+			SmallestInterval = TestInterval;
+		}
+		else
+		{
+			// Since our samples are stored in chronological order, it makes no sense
+			// to continue searching. The interval will continue increasing.
+			break;
+		}
+	}
+
+	checkSlow(ClosestIndex >= 0 && ClosestIndex < TextureSamples.Num());
+
+	// Finally, return the closest sample we found
+	return StaticCastSharedPtr<FMediaIOCoreTextureSampleBase, IMediaTextureSample, ESPMode::ThreadSafe>(TextureSamples[ClosestIndex]);
+}
+
+TSharedPtr<FMediaIOCoreTextureSampleBase> FMediaIOCorePlayerBase::PickSampleToRenderFramelocked_RenderThread(const TSharedPtr<FMediaIOCoreTextureSampleBase>& JITRProxySample)
+{
+	unimplemented();
+	return PickSampleToRenderForTimeSynchronized_RenderThread(JITRProxySample);
+}
+
+void FMediaIOCorePlayerBase::TransferTexture_RenderThread(const TSharedPtr<FMediaIOCoreTextureSampleBase>& Sample, const TSharedPtr<FMediaIOCoreTextureSampleBase>& JITRProxySample)
+{
+	checkSlow(IsInRenderingThread());
+
+	FRHICommandListImmediate& RHICmdList = FRHICommandListExecutor::GetImmediateCommandList();
+
+	// DMA P2P transfer
+	const bool bIsSampleAwaitingForGpuTransfer = Sample->IsAwaitingForGPUTransfer();
+	if (CanUseGPUTextureTransfer() && bIsSampleAwaitingForGpuTransfer)
+	{
+		SCOPED_GPU_STAT(RHICmdList, STAT_MediaIOPlayer_JITR_TransferTexture);
+		SCOPED_DRAW_EVENT(RHICmdList, STAT_MediaIOPlayer_JITR_TransferTexture);
+
+		// Prepare the proxy sample for DMA texture transfer
+		PreGPUTransferJITR(Sample, JITRProxySample);
+
+		// Perform fast texture transfer to an external registered texture
+		ExecuteGPUTransfer(JITRProxySample);
+	}
+	// Otherwise set raw data buffer so the texture can be built from it by the caller.
+	else
+	{
+		// Since the proxy sample holds a TSharedPtr reference to this original sample,
+		// we can be sure the internal buffer of the source sample won't be released prematurely.
+		JITRProxySample->SetBuffer(Sample->GetMutableBuffer());
+	}
+}
+
+bool FMediaIOCorePlayerBase::FJITRMediaTextureSamples::FetchVideo(TRange<FTimespan> TimeRange, TSharedPtr<IMediaTextureSample>& OutSample)
+{
+	if (!ProxySample.IsValid())
+	{
+		return false;
+	}
+
+	OutSample = ProxySample;
+
+	ProxySample.Reset();
+
+	return true;
+}
+
+bool FMediaIOCorePlayerBase::FJITRMediaTextureSamples::PeekVideoSampleTime(FMediaTimeStamp& TimeStamp)
+{
+	return false;
+}
+
+void FMediaIOCorePlayerBase::FJITRMediaTextureSamples::FlushSamples()
+{
+	ProxySample.Reset();
+}
+
+void FMediaIOCorePlayerBase::LogBookmark(const FString& Text, const TSharedRef<IMediaTextureSample>& Sample)
+{
+	// Put a bookmark if requested
+	const bool bDetailedInsights = MediaIOCorePlayerDetail::CVarMediaIOJITRInsights.GetValueOnAnyThread();
+	if (bDetailedInsights)
+	{
+		const TOptional<FTimecode> Timecode = Sample->GetTimecode();
+		if (Timecode.IsSet())
+		{
+			const FFrameNumber FrameNumber = Timecode.GetValue().ToFrameNumber(VideoFrameRate);
+			TRACE_BOOKMARK(TEXT("%s [%d]"), *Text, FrameNumber.Value % 100);
+		}
+		else
+		{
+			const int32 Milliseconds = Sample->GetTime().Time.GetFractionMilli();
+			TRACE_BOOKMARK(TEXT("%s [%d]"), *Text, Milliseconds);
+		}
+	}
+}
+
 #undef LOCTEXT_NAMESPACE
diff --git a/Engine/Plugins/Media/MediaIOFramework/Source/MediaIOCore/Private/Shared/MediaIOCoreTextureSampleBase.cpp b/Engine/Plugins/Media/MediaIOFramework/Source/MediaIOCore/Private/Shared/MediaIOCoreTextureSampleBase.cpp
--- a/Engine/Plugins/Media/MediaIOFramework/Source/MediaIOCore/Private/Shared/MediaIOCoreTextureSampleBase.cpp
+++ b/Engine/Plugins/Media/MediaIOFramework/Source/MediaIOCore/Private/Shared/MediaIOCoreTextureSampleBase.cpp
@@ -2,6 +2,8 @@
 
 #include "MediaIOCoreTextureSampleBase.h"
 
+#include "MediaIOCoreTextureSampleConverter.h"
+
 #include "RHI.h"
 #include "RHIResources.h"
 
@@ -146,7 +148,60 @@ void* FMediaIOCoreTextureSampleBase::RequestBuffer(uint32 InBufferSize)
 	return Buffer.GetData();
 }
 
+TSharedPtr<FMediaIOCorePlayerBase> FMediaIOCoreTextureSampleBase::GetPlayer() const
+{
+	return Player.Pin();
+}
+
+bool FMediaIOCoreTextureSampleBase::InitializeJITR(const FMediaIOCoreSampleJITRConfigurationArgs& Args)
+{
+	if (!Args.Player|| !Args.Converter)
+	{
+		return false;
+	}
+
+	// Native sample data
+	Width  = Args.Width;
+	Height = Args.Height;
+	Time   = Args.Time;
+	Timecode = Args.Timecode;
+
+	// JITR data
+	Player    = Args.Player;
+	Converter = Args.Converter;
+	EvaluationOffsetInSeconds = Args.EvaluationOffsetInSeconds;
+
+	return true;
+}
+
+void FMediaIOCoreTextureSampleBase::CopyConfiguration(const TSharedPtr<FMediaIOCoreTextureSampleBase>& SourceSample)
+{
+	if (!SourceSample.IsValid())
+	{
+		return;
+	}
+
+	// Copy configuration parameters
+	Stride = SourceSample->Stride;
+	Width  = SourceSample->Width;
+	Height = SourceSample->Height;
+	SampleFormat = SourceSample->SampleFormat;
+	Time = SourceSample->Time;
+	Timecode = SourceSample->Timecode;
+	Encoding = SourceSample->Encoding;
+	ColorSpace = SourceSample->ColorSpace;
+	ColorSpaceStruct = SourceSample->ColorSpaceStruct;
+
+	// Save original sample
+	OriginalSample = SourceSample;
+}
+
 #if WITH_ENGINE
+IMediaTextureSampleConverter* FMediaIOCoreTextureSampleBase::GetMediaTextureSampleConverter()
+{
+	return Converter.Get();
+}
+
 FRHITexture* FMediaIOCoreTextureSampleBase::GetTexture() const
 {
 	return Texture.GetReference();
diff --git a/Engine/Plugins/Media/MediaIOFramework/Source/MediaIOCore/Public/CaptureCardMediaSource.h b/Engine/Plugins/Media/MediaIOFramework/Source/MediaIOCore/Public/CaptureCardMediaSource.h
--- a/Engine/Plugins/Media/MediaIOFramework/Source/MediaIOCore/Public/CaptureCardMediaSource.h
+++ b/Engine/Plugins/Media/MediaIOFramework/Source/MediaIOCore/Public/CaptureCardMediaSource.h
@@ -17,6 +17,9 @@
 
 namespace UE::CaptureCardMediaSource
 {
+	static FName RenderJIT("RenderJIT");
+	static FName Framelock("Framelock");
+	static FName EvaluationType("EvaluationType");
 	static FName Deinterlacer("Deinterlacer");
 	static FName InterlaceFieldOrder("InterlaceFieldOrder");
 	static FName OverrideSourceEncoding("OverrideSourceEncoding");
@@ -47,12 +50,22 @@ class MEDIAIOCORE_API UCaptureCardMediaSource : public UTimeSynchronizableMediaS
 	GENERATED_BODY()
 	
 public:
-	
-	UCaptureCardMediaSource()
-	{
-		Deinterlacer = CreateDefaultSubobject<UBobDeinterlacer>("Deinterlacer");
-	}
-	
+
+	UCaptureCardMediaSource();
+
+public:
+	/** Should use JITR technique? It enables late sample picking for render. */
+	UPROPERTY(EditAnywhere, AdvancedDisplay, Category = Synchronization, meta = (DisplayName = "Just-In-Time Rendering"))
+	bool bRenderJIT = true;
+
+	/** Should wait for some time until requested frame arrives? Requires JIT rendering. */
+	UPROPERTY(EditAnywhere, AdvancedDisplay, Category = Synchronization, meta = (DisplayName = "Framelock"))
+	bool bFramelock = false;
+
+	/** Sample evaluation type. */
+	UPROPERTY(EditAnywhere, AdvancedDisplay, Category = Synchronization, meta = (DisplayName = "Sample Evaluation Type"))
+	EMediaIOSampleEvaluationType EvaluationType = EMediaIOSampleEvaluationType::PlatformTime;
+
 	/**
 	 * How interlaced video should be treated.
 	 */
@@ -88,4 +101,10 @@ public:
 	virtual int64 GetMediaOption(const FName& Key, int64 DefaultValue) const override;
 	virtual bool GetMediaOption(const FName& Key, bool bDefaultValue) const override;
 	virtual bool HasMediaOption(const FName& Key) const override;
+
+public:
+#if WITH_EDITOR
+	virtual bool CanEditChange(const FProperty* InProperty) const override;
+	virtual void PostEditChangeChainProperty(struct FPropertyChangedChainEvent& PropertyChangedEvent) override;
+#endif //WITH_EDITOR
 };
diff --git a/Engine/Plugins/Media/MediaIOFramework/Source/MediaIOCore/Public/MediaIOCoreBinarySampleBase.h b/Engine/Plugins/Media/MediaIOFramework/Source/MediaIOCore/Public/MediaIOCoreBinarySampleBase.h
--- a/Engine/Plugins/Media/MediaIOFramework/Source/MediaIOCore/Public/MediaIOCoreBinarySampleBase.h
+++ b/Engine/Plugins/Media/MediaIOFramework/Source/MediaIOCore/Public/MediaIOCoreBinarySampleBase.h
@@ -6,7 +6,7 @@
 #include "MediaObjectPool.h"
 
 /**
- * Implements a media binary data sample for AjaMedia.
+ * Implements a media binary data sample.
  */
 class MEDIAIOCORE_API FMediaIOCoreBinarySampleBase
 	: public IMediaBinarySample
diff --git a/Engine/Plugins/Media/MediaIOFramework/Source/MediaIOCore/Public/MediaIOCoreCaptionSampleBase.h b/Engine/Plugins/Media/MediaIOFramework/Source/MediaIOCore/Public/MediaIOCoreCaptionSampleBase.h
new file mode 100644
--- /dev/null
+++ b/Engine/Plugins/Media/MediaIOFramework/Source/MediaIOCore/Public/MediaIOCoreCaptionSampleBase.h
@@ -0,0 +1,16 @@
+// Copyright Epic Games, Inc. All Rights Reserved.
+
+#pragma once
+
+#include "IMediaOverlaySample.h"
+#include "MediaObjectPool.h"
+
+
+/**
+ * Implements a media caption sample.
+ */
+class MEDIAIOCORE_API FMediaIOCoreCaptionSampleBase
+	: public IMediaOverlaySample
+	, public IMediaPoolable
+{
+};
diff --git a/Engine/Plugins/Media/MediaIOFramework/Source/MediaIOCore/Public/MediaIOCoreDefinitions.h b/Engine/Plugins/Media/MediaIOFramework/Source/MediaIOCore/Public/MediaIOCoreDefinitions.h
--- a/Engine/Plugins/Media/MediaIOFramework/Source/MediaIOCore/Public/MediaIOCoreDefinitions.h
+++ b/Engine/Plugins/Media/MediaIOFramework/Source/MediaIOCore/Public/MediaIOCoreDefinitions.h
@@ -76,6 +76,17 @@ enum class EMediaIOInterlaceFieldOrder : uint8
 	BottomFieldFirst
 };
 
+/**
+ * Various sample evaluation types availalbe in MediaIO.
+ */
+UENUM()
+enum class EMediaIOSampleEvaluationType : uint8
+{
+	Latest,
+	PlatformTime,
+	Timecode,
+};
+
 namespace UE::MediaIO
 {
 	static EMediaIOAutoDetectableTimecodeFormat ToAutoDetectableTimecodeFormat(EMediaIOTimecodeFormat TimecodeFormat)
diff --git a/Engine/Plugins/Media/MediaIOFramework/Source/MediaIOCore/Public/MediaIOCorePlayerBase.h b/Engine/Plugins/Media/MediaIOFramework/Source/MediaIOCore/Public/MediaIOCorePlayerBase.h
--- a/Engine/Plugins/Media/MediaIOFramework/Source/MediaIOCore/Public/MediaIOCorePlayerBase.h
+++ b/Engine/Plugins/Media/MediaIOFramework/Source/MediaIOCore/Public/MediaIOCorePlayerBase.h
@@ -11,14 +11,20 @@
 #include "IMediaView.h"
 #include "ITimedDataInput.h"
 
-#include "HAL/CriticalSection.h"
+#include "MediaIOCoreDefinitions.h"
 #include "MediaIOCoreSampleContainer.h"
-#include "MediaIOCoreTextureSampleBase.h"
-#include "Misc/CoreMisc.h"
 
+#include "HAL/CriticalSection.h"
+#include "Misc/CoreMisc.h"
 #include "Misc/FrameRate.h"
 
+class FMediaIOCoreAudioSampleBase;
+class FMediaIOCoreBinarySampleBase;
+class FMediaIOCoreCaptionSampleBase;
 class FMediaIOCoreSamples;
+class FMediaIOCoreSubtitleSampleBase;
+class FMediaIOCoreTextureSampleBase;
+class FMediaIOCoreTextureSampleConverter;
 class IMediaEventSink;
 
 enum class EMediaTextureSampleFormat;
@@ -153,6 +159,31 @@ public:
 	virtual void RemoveChannel(ITimedDataInputChannel * Channel) override;
 	virtual bool SupportsSubFrames() const override;
 
+public:
+
+	/**
+	 * Just in time sample rendering. This method is responsible for late sample picking,
+	 * then rendering it into the proxy sample provided.
+	 */
+	bool JustInTimeSampleRender_RenderThread(TSharedPtr<FMediaIOCoreTextureSampleBase>& JITRProxySample);
+
+protected:
+
+	/** Internal API to push an audio sample to the internal samples pool */
+	void AddAudioSample(const TSharedRef<FMediaIOCoreAudioSampleBase>& Sample);
+
+	/** Internal API to push a caption sample to the internal samples pool */
+	void AddCaptionSample(const TSharedRef<FMediaIOCoreCaptionSampleBase>& Sample);
+
+	/** Internal API to push a metadata sample to the internal samples pool */
+	void AddMetadataSample(const TSharedRef<FMediaIOCoreBinarySampleBase>& Sample);
+
+	/** Internal API to push a subtitle sample to the internal samples pool */
+	void AddSubtitleSample(const TSharedRef<FMediaIOCoreSubtitleSampleBase>& Sample);
+
+	/** Internal API to push a video sample to the internal samples pool */
+	void AddVideoSample(const TSharedRef<FMediaIOCoreTextureSampleBase>& Sample);
+
 protected:
 	/** Is the IO hardware/device ready to be used. */
 	virtual bool IsHardwareReady() const = 0;
@@ -160,6 +191,12 @@ protected:
 	/** Return true if the options combination are valid. */
 	virtual bool ReadMediaOptions(const IMediaOptions* Options);
 
+	/**
+	 * Allows children to notify when video format is known after auto-detection. The main purpose
+	 * is to start intialization of the DMA textures at the right moment.
+	 */
+	void NotifyVideoFormatDetected();
+
 	/** Get the application time with a delta to represent the actual system time. Use instead of FApp::GetCurrentTime. */
 	static double GetApplicationSeconds();
 
@@ -188,95 +225,158 @@ protected:
 		return EMediaIOCoreColorFormat::YUV8;
 	}
 
-	virtual void AddVideoSample(const TSharedRef<FMediaIOCoreTextureSampleBase>& InSample)
-	{
-	}
+	/** Called after fast GPUDirect texture transferring is finished */
+	virtual void AddVideoSampleAfterGPUTransfer_RenderThread(const TSharedRef<FMediaIOCoreTextureSampleBase>& Sample);
 
-	virtual bool CanUseGPUTextureTransfer();
+	/** Whether fast GPU transferring is available */
+	virtual bool CanUseGPUTextureTransfer() const;
 
-	/** Handler used to provide texture samples to the deinterlacer. */
-	virtual TSharedPtr<FMediaIOCoreTextureSampleBase> AcquireSample_AnyThread() const
-	{
-		ensureMsgf(false, TEXT("If deinterlacing is used by the media player, AcquireSample_AnyThread must be implemented by the child media player class."));
-		return nullptr;
-	}
+	/** Whether a texture availabled for GPU transfer */
+	bool HasTextureAvailableForGPUTransfer() const;
 
-	void OnSampleDestroyed(TRefCountPtr<FRHITexture> InTexture);
-	void RegisterSampleBuffer(const TSharedPtr<FMediaIOCoreTextureSampleBase>& InSample);
-	void UnregisterSampleBuffers();
-	void CreateAndRegisterTextures(const IMediaOptions* Options);
-	void UnregisterTextures();
-	void PreGPUTransfer(const TSharedPtr<FMediaIOCoreTextureSampleBase>& InSample);
-	void ExecuteGPUTransfer(const TSharedPtr<FMediaIOCoreTextureSampleBase>& InSample);
+	/** Whether JITR is available and active */
+	virtual bool IsJustInTimeRenderingEnabled() const;
+
+	/** Factory method to generate texture samples. */
+	virtual TSharedPtr<FMediaIOCoreTextureSampleBase> AcquireTextureSample_AnyThread() const = 0;
 
 	/** Deinterlace a video frame. Only valid to call when the video stream is open. */
 	TArray<TSharedRef<FMediaIOCoreTextureSampleBase>> Deinterlace(const UE::MediaIOCore::FVideoFrame& InVideoFrame) const;
 
 protected:
-	/** Critical section for synchronizing access to receiver and sinks. */
-	FCriticalSection CriticalSection;
 
+	/** Acquire a proxy sample for JITR. That sample must be initialized for JITR. */
+	virtual TSharedPtr<FMediaIOCoreTextureSampleBase> AcquireJITRProxySampleInitialized();
+
+	/** Factory method to create sample converters. */
+	virtual TSharedPtr<FMediaIOCoreTextureSampleConverter> CreateTextureSampleConverter() const;
+
+	/** Pick a sample to render */
+	TSharedPtr<FMediaIOCoreTextureSampleBase> PickSampleToRender_RenderThread(const TSharedPtr<FMediaIOCoreTextureSampleBase>& JITRProxySample);
+	TSharedPtr<FMediaIOCoreTextureSampleBase> PickSampleToRenderForLatest_RenderThread(const TSharedPtr<FMediaIOCoreTextureSampleBase>& JITRProxySample);
+	TSharedPtr<FMediaIOCoreTextureSampleBase> PickSampleToRenderForTimeSynchronized_RenderThread(const TSharedPtr<FMediaIOCoreTextureSampleBase>& JITRProxySample);
+	TSharedPtr<FMediaIOCoreTextureSampleBase> PickSampleToRenderFramelocked_RenderThread(const TSharedPtr<FMediaIOCoreTextureSampleBase>& JITRProxySample);
+
+	/**
+	 * A wrapper method responsible for transferring of the sample textures into GPU memory based
+	 * on the current settings and hardware capabilities.
+	 */
+	void TransferTexture_RenderThread(const TSharedPtr<FMediaIOCoreTextureSampleBase>& Sample, const TSharedPtr<FMediaIOCoreTextureSampleBase>& JITRProxySample);
+
+protected:
 	/** Url used to open the media player. */
 	FString OpenUrl;
 
 	/** format of the video. */
 	FMediaVideoTrackFormat VideoTrackFormat;
-	
+
 	/** format of the audio. */
 	FMediaAudioTrackFormat AudioTrackFormat;
 
 	/** Current state of the media player. */
-	EMediaState CurrentState;
+	EMediaState CurrentState = EMediaState::Closed;
 
 	/** Current playback time. */
-	FTimespan CurrentTime;
+	FTimespan CurrentTime = FTimespan::Zero();
 
 	/** The media event handler. */
 	IMediaEventSink& EventSink;
 
 	/** Video frame rate in the last received sample. */
-	FFrameRate VideoFrameRate;
+	FFrameRate VideoFrameRate = { 30, 1 };
 
 	/** The media sample cache. */
-	FMediaIOCoreSamples* Samples;
+	const TUniquePtr<FMediaIOCoreSamples> Samples;
 
-	/** Whether to use the Synchronization Time module as time source. */
-	bool bUseTimeSynchronization;
+	/** Sample evaluation type. */
+	EMediaIOSampleEvaluationType EvaluationType = EMediaIOSampleEvaluationType::PlatformTime;
 
 	/** Warn when the video frame rate is not the same as the engine's frame rate. */
-	bool bWarnedIncompatibleFrameRate;
+	bool bWarnedIncompatibleFrameRate = false;
 
 	/** Whether we are using autodetection. */
-	bool bAutoDetect;
+	bool bAutoDetect = false;
 
 	/** When using Time Synchronization (TC synchronization), how many frame back of a delay would you like. */
-	int32 FrameDelay;
+	int32 FrameDelay = 0;
 
 	/** When not using Time Synchronization (use computer time), how many sec back of a delay would you like. */
-	double TimeDelay;
+	double TimeDelay = 0.0;
 
 	/** Previous frame Timespan */
 	FTimespan PreviousFrameTimespan;
 
 	/** Timed Data Input handler */
 	TArray<ITimedDataInputChannel*> Channels;
-	
+
 	/** Base set of settings to start from when setuping channels */
 	FMediaIOSamplingSettings BaseSettings;
 
-	FCriticalSection TexturesCriticalSection;
-
-	/** Pool of textures registerd with GPU Texture transfer. */
-	TArray<TRefCountPtr<FRHITexture>> Textures;
-
 private:
+	void OnSampleDestroyed(TRefCountPtr<FRHITexture> InTexture);
+	void RegisterSampleBuffer(const TSharedPtr<FMediaIOCoreTextureSampleBase>& InSample);
+	void UnregisterSampleBuffers();
+	void CreateAndRegisterTextures();
+	void UnregisterTextures();
+	void PreGPUTransfer(const TSharedPtr<FMediaIOCoreTextureSampleBase>& InSample);
+	void PreGPUTransferJITR(const TSharedPtr<FMediaIOCoreTextureSampleBase>& InSample, const TSharedPtr<FMediaIOCoreTextureSampleBase>& InJITRProxySample);
+	void ExecuteGPUTransfer(const TSharedPtr<FMediaIOCoreTextureSampleBase>& InSample);
+
 	/** GPU Texture transfer object */
 	UE::GPUTextureTransfer::TextureTransferPtr GPUTextureTransfer;
+
 	/** Buffers Registered with GPU Texture Transfer */
 	TSet<void*> RegisteredBuffers;
-	/** Pool of textures registerd with GPU Texture transfer. */
+
+	/** Textures registerd with GPU Texture transfer. */
 	TSet<TRefCountPtr<FRHITexture>> RegisteredTextures;
+
+	/** Pool of textures available for GPU Texture transfer. */
+	TArray<TRefCountPtr<FRHITexture>> Textures;
+
+	/** Critical section to control access to the pool. */
+	mutable FCriticalSection TexturesCriticalSection;
+
 	/** Utility object to handle deinterlacing. */
 	TSharedPtr<UE::MediaIOCore::FDeinterlacer> Deinterlacer;
 
+private:
+
+	/** Utility function to log various sample events (i.e. received/transferred/picked) */
+	void LogBookmark(const FString& Text, const TSharedRef<IMediaTextureSample>& Sample);
+
+private:
+
+	/** Special JITR sample container, holding a single sample used to trigger media framework rendering */
+	class FJITRMediaTextureSamples
+		: public IMediaSamples
+	{
+	public:
+		FJITRMediaTextureSamples() = default;
+		FJITRMediaTextureSamples(const FJITRMediaTextureSamples&) = delete;
+		FJITRMediaTextureSamples& operator=(const FJITRMediaTextureSamples&) = delete;
+
+	public:
+		//~ Begin IMediaSamples interface
+		virtual bool FetchVideo(TRange<FTimespan> TimeRange, TSharedPtr<IMediaTextureSample>& OutSample) override;
+		virtual bool PeekVideoSampleTime(FMediaTimeStamp& TimeStamp) override;
+		virtual void FlushSamples() override;
+		//~ End IMediaSamples interface
+
+	public:
+		// JITR sample
+		TSharedPtr<FMediaIOCoreTextureSampleBase> ProxySample;
+	};
+
+	/** Is Just-In-Time Rendering enabled */
+	bool bJustInTimeRender = false;
+
+	/** Whether media playback should be framelocked to the engine's timecode */
+	bool bFramelock = false;
+
+	/** JITR samples proxy */
+	const TUniquePtr<FJITRMediaTextureSamples> JITRSamples;
+
+	/** Used to ensure that JIT rendering is executed only once per frame */
+	uint64 LastEngineRTFrameThatUpdatedJustInTime = TNumericLimits<uint64>::Max();
 };
diff --git a/Engine/Plugins/Media/MediaIOFramework/Source/MediaIOCore/Public/MediaIOCoreSampleContainer.h b/Engine/Plugins/Media/MediaIOFramework/Source/MediaIOCore/Public/MediaIOCoreSampleContainer.h
--- a/Engine/Plugins/Media/MediaIOFramework/Source/MediaIOCore/Public/MediaIOCoreSampleContainer.h
+++ b/Engine/Plugins/Media/MediaIOFramework/Source/MediaIOCore/Public/MediaIOCoreSampleContainer.h
@@ -368,6 +368,13 @@ public:
 		return false;
 	}
 
+	/** Returns all available samples */
+	TArray<TSharedPtr<SampleType, ESPMode::ThreadSafe>> GetSamples() const
+	{
+		FScopeLock Lock(&CriticalSection);
+		return Samples;
+	}
+
 	void FlushSamples()
 	{
 		FScopeLock Lock(&CriticalSection);
@@ -408,7 +415,3 @@ protected:
 	mutable FCriticalSection CriticalSection;
 	TArray<TSharedPtr<SampleType, ESPMode::ThreadSafe>> Samples;
 };
-
-
-
-
diff --git a/Engine/Plugins/Media/MediaIOFramework/Source/MediaIOCore/Public/MediaIOCoreSamples.h b/Engine/Plugins/Media/MediaIOFramework/Source/MediaIOCore/Public/MediaIOCoreSamples.h
--- a/Engine/Plugins/Media/MediaIOFramework/Source/MediaIOCore/Public/MediaIOCoreSamples.h
+++ b/Engine/Plugins/Media/MediaIOFramework/Source/MediaIOCore/Public/MediaIOCoreSamples.h
@@ -272,6 +272,14 @@ public:
 		return CaptionSamples.GetFrameDroppedStat();
 	}
 
+	/**
+	 * Returns all available video samples
+	 */
+	TArray<TSharedPtr<IMediaTextureSample, ESPMode::ThreadSafe>> GetVideoSamples() const
+	{
+		return VideoSamples.GetSamples();
+	}
+
 	/**
 	 * Caches the current sample container state for a given Player (evaluation) time
 	 */
diff --git a/Engine/Plugins/Media/MediaIOFramework/Source/MediaIOCore/Public/MediaIOCoreSubtitleSampleBase.h b/Engine/Plugins/Media/MediaIOFramework/Source/MediaIOCore/Public/MediaIOCoreSubtitleSampleBase.h
new file mode 100644
--- /dev/null
+++ b/Engine/Plugins/Media/MediaIOFramework/Source/MediaIOCore/Public/MediaIOCoreSubtitleSampleBase.h
@@ -0,0 +1,16 @@
+// Copyright Epic Games, Inc. All Rights Reserved.
+
+#pragma once
+
+#include "IMediaOverlaySample.h"
+#include "MediaObjectPool.h"
+
+
+/**
+ * Implements a media subtitle sample.
+ */
+class MEDIAIOCORE_API FMediaIOCoreSubtitleSampleBase
+	: public IMediaOverlaySample
+	, public IMediaPoolable
+{
+};
diff --git a/Engine/Plugins/Media/MediaIOFramework/Source/MediaIOCore/Public/MediaIOCoreTextureSampleBase.h b/Engine/Plugins/Media/MediaIOFramework/Source/MediaIOCore/Public/MediaIOCoreTextureSampleBase.h
--- a/Engine/Plugins/Media/MediaIOFramework/Source/MediaIOCore/Public/MediaIOCoreTextureSampleBase.h
+++ b/Engine/Plugins/Media/MediaIOFramework/Source/MediaIOCore/Public/MediaIOCoreTextureSampleBase.h
@@ -10,6 +10,8 @@
 #include "Misc/FrameRate.h"
 #include "Templates/RefCounting.h"
 
+class FMediaIOCorePlayerBase;
+class FMediaIOCoreTextureSampleConverter;
 class FRHITexture;
 
 namespace UE::MediaIOCore
@@ -49,6 +51,33 @@ namespace UE::MediaIOCore
 }
 
 
+
+/** Used to setup JITR data for a sample */
+struct MEDIAIOCORE_API FMediaIOCoreSampleJITRConfigurationArgs
+{
+	/** Width of the sample in pixels */
+	uint32 Width = 0;
+
+	/** Height of the sample in pixels */
+	uint32 Height = 0;
+
+	/** Platform time fixed on the game thread for JITR sample picking */
+	FTimespan Time;
+
+	/** Engine timecode fixed on the game thread for JITR sample picking */
+	FTimecode Timecode;
+
+	/** Time offset evaluated on game thread */
+	double EvaluationOffsetInSeconds = 0;
+
+	/** Player that produced the sample */
+	TSharedPtr<FMediaIOCorePlayerBase> Player;
+
+	/** Sample converter to process this sample */
+	TSharedPtr<FMediaIOCoreTextureSampleConverter> Converter;
+};
+
+
 /**
  * Implements the IMediaTextureSample/IMediaPoolable interface.
  */
@@ -176,8 +205,6 @@ public:
 	 */
 	bool SetBufferWithEvenOddLine(bool bUseEvenLine, const void* InVideoBuffer, uint32 InBufferSize, uint32 InStride, uint32 InHeight);
 
-
-
 	/**
 	 * Request an uninitialized sample buffer.
 	 * Should be used when the buffer could be filled by something else.
@@ -187,6 +214,37 @@ public:
 	 */
 	virtual void* RequestBuffer(uint32 InBufferSize);
 
+	/**
+	 * Configure this sample for JITR
+	 *
+	 * @param Args JITR configuration parameters.
+	 */
+	virtual bool InitializeJITR(const FMediaIOCoreSampleJITRConfigurationArgs& Args);
+
+	/** Marks this sample as one that is ready and awaiting for fast GPUDirect texture transfer */
+	void SetAwaitingForGPUTransfer()
+	{
+		bIsAwaitingForGPUTransfer = true;
+	}
+
+	/** Returns whether it's ready for GPUDirect texture transfer */
+	bool IsAwaitingForGPUTransfer() const
+	{
+		return bIsAwaitingForGPUTransfer;
+	}
+
+	/** Returns the player that created this sample */
+	TSharedPtr<FMediaIOCorePlayerBase> GetPlayer() const;
+
+	/** Returns time evaluated on the game thread for JITR */
+	double GetEvaluationOffsetInSeconds() const
+	{
+		return EvaluationOffsetInSeconds;
+	}
+
+	/** Copies all neccessary data from a source sample to render JIT */
+	virtual void CopyConfiguration(const TSharedPtr<FMediaIOCoreTextureSampleBase>& SourceSample);
+
 public:
 	//~ IMediaTextureSample interface
 
@@ -283,6 +341,7 @@ public:
 	}
 
 #if WITH_ENGINE
+	virtual IMediaTextureSampleConverter* GetMediaTextureSampleConverter() override;
 	virtual FRHITexture* GetTexture() const override;
 #endif //WITH_ENGINE
 
@@ -310,6 +369,8 @@ protected:
 	virtual void FreeSample()
 	{
 		Buffer.Reset();
+		ExternalBuffer = nullptr;
+		Texture.SafeRelease();
 	}
 
 	/**
@@ -333,9 +394,9 @@ protected:
 	TOptional<FTimecode> Timecode;
 
 	/** Image dimensions */
-	uint32 Stride;
-	uint32 Width;
-	uint32 Height;
+	uint32 Stride = 0;
+	uint32 Width  = 0;
+	uint32 Height = 0;
 
 	/** Pointer to raw pixels */
 	TArray<uint8, TAlignedHeapAllocator<4096>> Buffer;
@@ -355,4 +416,23 @@ protected:
 
 	/** Color space structure of the incoming texture. Used for retrieving chromaticities. */
 	UE::Color::FColorSpace ColorSpaceStruct = UE::Color::FColorSpace(UE::Color::EColorSpace::sRGB);
+
+private:
+	/** The player that created this sample */
+	TWeakPtr<FMediaIOCorePlayerBase> Player;
+
+	/** Custom converter that will be the one checking back with the player for just in time sample render purposes */
+	TSharedPtr<FMediaIOCoreTextureSampleConverter> Converter;
+
+	/**
+	 * A reference to the original sample that was chosen during JITR. The idea of this member is to keep
+	 * the original sample alive, prevent any of its resources from being released while this proxy sample is in use.
+	 */
+	TSharedPtr<FMediaIOCoreTextureSampleBase> OriginalSample;
+
+	/** Whether this sample's texture data is awaiting to be transferred by GPUDirect */
+	bool bIsAwaitingForGPUTransfer = false;
+
+	/** Time offset evaluated on game thread for JITR */
+	double EvaluationOffsetInSeconds = 0;
 };
diff --git a/Engine/Plugins/Media/MediaIOFramework/Source/MediaIOCore/Public/MediaIOCoreTextureSampleConverter.h b/Engine/Plugins/Media/MediaIOFramework/Source/MediaIOCore/Public/MediaIOCoreTextureSampleConverter.h
new file mode 100644
--- /dev/null
+++ b/Engine/Plugins/Media/MediaIOFramework/Source/MediaIOCore/Public/MediaIOCoreTextureSampleConverter.h
@@ -0,0 +1,40 @@
+// Copyright Epic Games, Inc. All Rights Reserved.
+
+#pragma once
+
+#include "IMediaTextureSampleConverter.h"
+
+class FMediaIOCoreTextureSampleBase;
+
+
+/**
+ * Media IO base texture sample converter.
+ *
+ * It's mostly responsible for Just-In-Time Rendering (JITR), but also provides a way
+ * to implement custom sample conversion when inherited.
+ * 
+ * @note: Don't forget to call FMediaIOTextureSampleConverter::Convert(...) from the child classes
+ *        to trigger JITR pipeline. This is the place where late sample picking is actually happens.
+ */
+class MEDIAIOCORE_API FMediaIOCoreTextureSampleConverter
+	: public IMediaTextureSampleConverter
+{
+public:
+	FMediaIOCoreTextureSampleConverter() = default;
+	virtual ~FMediaIOCoreTextureSampleConverter() = default;
+
+public:
+	/** Configures settings to convert incoming sample */
+	virtual void Setup(const TSharedPtr<FMediaIOCoreTextureSampleBase>& InSample);
+
+public:
+	//~ Begin IMediaTextureSampleConverter interface
+	virtual bool Convert(FTexture2DRHIRef& InDstTexture, const FConversionHints& Hints) override;
+	virtual uint32 GetConverterInfoFlags() const override;
+	//~ End IMediaTextureSampleConverter interface
+
+private:
+
+	/** Proxy sample for JITR */
+	TWeakPtr<FMediaIOCoreTextureSampleBase> JITRProxySample;
+};
diff --git a/Engine/Plugins/VirtualProduction/Rivermax/RivermaxMedia/Source/RivermaxMedia/Private/RivermaxMediaPlayer.h b/Engine/Plugins/VirtualProduction/Rivermax/RivermaxMedia/Source/RivermaxMedia/Private/RivermaxMediaPlayer.h
--- a/Engine/Plugins/VirtualProduction/Rivermax/RivermaxMedia/Source/RivermaxMedia/Private/RivermaxMediaPlayer.h
+++ b/Engine/Plugins/VirtualProduction/Rivermax/RivermaxMedia/Source/RivermaxMedia/Private/RivermaxMediaPlayer.h
@@ -88,6 +88,12 @@ namespace  UE::RivermaxMedia
 		//~ Begin FMediaIOCorePlayerBase interface
 		virtual bool IsHardwareReady() const override;
 		virtual void SetupSampleChannels() override;
+		virtual TSharedPtr<FMediaIOCoreTextureSampleBase> AcquireTextureSample_AnyThread() const override
+		{
+			// This needs to be fixed once FRivermaxMediaTextureSample is inherited from FMediaIOCoreTextureSampleBase
+			// return MakeShared<FRivermaxMediaTextureSample>();
+			return nullptr;
+		}
 		//~ End FMediaIOCorePlayerBase interface
 
 	private:
diff --git a/Engine/Source/Runtime/MediaAssets/Public/TimeSynchronizableMediaSource.h b/Engine/Source/Runtime/MediaAssets/Public/TimeSynchronizableMediaSource.h
--- a/Engine/Source/Runtime/MediaAssets/Public/TimeSynchronizableMediaSource.h
+++ b/Engine/Source/Runtime/MediaAssets/Public/TimeSynchronizableMediaSource.h
@@ -42,7 +42,7 @@ public:
 	 * The media player has be able to read timecode.
 	 * The media player will try to play the corresponding frame, base on the frame's timecode value.
 	 */
-	UPROPERTY(BlueprintReadWrite, EditAnywhere, Category=Synchronization, meta=(DisplayName="Synchronize with Engine's Timecode"))
+	UPROPERTY(BlueprintReadWrite, EditAnywhere, Category=Synchronization, meta=(DisplayName="Time Synchronization"))
 	bool bUseTimeSynchronization;
 
 	/** When using Time Synchronization, how many frame back should it read. */
