Add OCIO support for aja and blackmagic media sources

diff --git a/Engine/Plugins/Media/AjaMedia/Source/AjaMedia/AjaMedia.Build.cs b/Engine/Plugins/Media/AjaMedia/Source/AjaMedia/AjaMedia.Build.cs
--- a/Engine/Plugins/Media/AjaMedia/Source/AjaMedia/AjaMedia.Build.cs
+++ b/Engine/Plugins/Media/AjaMedia/Source/AjaMedia/AjaMedia.Build.cs
@@ -25,6 +25,7 @@ namespace UnrealBuildTool.Rules
 					"Engine",
 					"MediaIOCore",
 					"MediaUtils",
+					"OpenColorIO",
 					"Projects",
 					"TimeManagement",
 					"RenderCore",
diff --git a/Engine/Plugins/Media/AjaMedia/Source/AjaMedia/Private/Player/AjaMediaPlayer.cpp b/Engine/Plugins/Media/AjaMedia/Source/AjaMedia/Private/Player/AjaMediaPlayer.cpp
--- a/Engine/Plugins/Media/AjaMedia/Source/AjaMedia/Private/Player/AjaMediaPlayer.cpp
+++ b/Engine/Plugins/Media/AjaMedia/Source/AjaMedia/Private/Player/AjaMediaPlayer.cpp
@@ -440,6 +440,7 @@ bool FAjaMediaPlayer::OnRequestInputBuffer(const AJA::AJARequestInputBufferData&
 
 		AjaThreadCurrentTextureSample = TextureSamplePool->AcquireShared();
 		AjaThreadCurrentTextureSample->RequestBuffer(InRequestBuffer.VideoBufferSize);
+		AjaThreadCurrentTextureSample->SetColorConversionSettings(OCIOSettings);
 
 		if (bIsJustInTimeRenderingEnabled && bCanUseGPUTextureTransfer)
 		{
@@ -699,6 +700,7 @@ bool FAjaMediaPlayer::OnInputFrameReceived(const AJA::AJAInputFrameData& InInput
 					auto TextureSample = TextureSamplePool->AcquireShared();
 					if (TextureSample->InitializeProgressive(InVideoFrame, VideoSampleFormat, DecodedTime, VideoFrameRate, DecodedTimecode, ColorFormat))
 					{
+						TextureSample->SetColorConversionSettings(OCIOSettings);
 						AddVideoSample(TextureSample);
 					}
 				}
@@ -727,6 +729,7 @@ bool FAjaMediaPlayer::OnInputFrameReceived(const AJA::AJAInputFrameData& InInput
 
 					for (const TSharedRef<FMediaIOCoreTextureSampleBase>& TextureSample : DeinterlacedSamples)
 					{
+						TextureSample->SetColorConversionSettings(OCIOSettings);
 						AddVideoSample(TextureSample);
 					}
 				}
diff --git a/Engine/Plugins/Media/BlackmagicMedia/Source/BlackmagicMedia/BlackmagicMedia.Build.cs b/Engine/Plugins/Media/BlackmagicMedia/Source/BlackmagicMedia/BlackmagicMedia.Build.cs
--- a/Engine/Plugins/Media/BlackmagicMedia/Source/BlackmagicMedia/BlackmagicMedia.Build.cs
+++ b/Engine/Plugins/Media/BlackmagicMedia/Source/BlackmagicMedia/BlackmagicMedia.Build.cs
@@ -34,6 +34,7 @@ namespace UnrealBuildTool.Rules
 					"BlackmagicCore",
 					"Engine",
 					"MediaUtils",
+					"OpenColorIO",
 					"Projects",
 					"SlateCore"
                 });
diff --git a/Engine/Plugins/Media/BlackmagicMedia/Source/BlackmagicMedia/Private/Player/BlackmagicMediaPlayer.cpp b/Engine/Plugins/Media/BlackmagicMedia/Source/BlackmagicMedia/Private/Player/BlackmagicMediaPlayer.cpp
--- a/Engine/Plugins/Media/BlackmagicMedia/Source/BlackmagicMedia/Private/Player/BlackmagicMediaPlayer.cpp
+++ b/Engine/Plugins/Media/BlackmagicMedia/Source/BlackmagicMedia/Private/Player/BlackmagicMediaPlayer.cpp
@@ -18,6 +18,7 @@
 #include "MediaIOCoreUtilities.h"
 #include "Misc/App.h"
 #include "Misc/ScopeLock.h"
+#include "OpenColorIOColorSpace.h"
 #include "RenderCommandFence.h"
 #include "Slate/SceneViewport.h"
 #include "Stats/Stats2.h"
@@ -367,6 +368,7 @@ namespace BlackmagicMediaPlayerHelpers
 						}
 
 						auto TextureSample = MediaPlayer->TextureSamplePool->AcquireShared();
+						TextureSample->SetColorConversionSettings(MediaPlayer->OCIOSettings);
 						bool bInitializeResult = false;
 						if (bGPUDirectTexturesAvailable)
 						{
@@ -424,6 +426,7 @@ namespace BlackmagicMediaPlayerHelpers
 
 						for (const TSharedRef<FMediaIOCoreTextureSampleBase>& TextureSample : DeinterlacedSamples)
 						{
+							TextureSample->SetColorConversionSettings(MediaPlayer->OCIOSettings);
 							MediaPlayer->AddVideoSample(TextureSample);
 						}
 					}
@@ -614,7 +617,7 @@ bool FBlackmagicMediaPlayer::Open(const FString& Url, const IMediaOptions* Optio
 	// Setup our different supported channels based on source settings
 	SetupSampleChannels();
 
-	bool bSuccess = EventCallback->Initialize(ChannelOptions, bEncodeTimecodeInTexel, MaxNumAudioFrameBuffer, MaxNumVideoFrameBuffer);
+	const bool bSuccess = EventCallback->Initialize(ChannelOptions, bEncodeTimecodeInTexel, MaxNumAudioFrameBuffer, MaxNumVideoFrameBuffer);
 
 	if (!bSuccess)
 	{
diff --git a/Engine/Plugins/Media/MediaIOFramework/Source/MediaIOCore/Private/CaptureCardMediaSource.cpp b/Engine/Plugins/Media/MediaIOFramework/Source/MediaIOCore/Private/CaptureCardMediaSource.cpp
--- a/Engine/Plugins/Media/MediaIOFramework/Source/MediaIOCore/Private/CaptureCardMediaSource.cpp
+++ b/Engine/Plugins/Media/MediaIOFramework/Source/MediaIOCore/Private/CaptureCardMediaSource.cpp
@@ -33,6 +33,18 @@ int64 UCaptureCardMediaSource::GetMediaOption(const FName& Key, int64 DefaultVal
 	return UTimeSynchronizableMediaSource::GetMediaOption(Key, DefaultValue);
 }
 
+TSharedPtr<IMediaOptions::FDataContainer, ESPMode::ThreadSafe> UCaptureCardMediaSource::GetMediaOption(const FName& Key, const TSharedPtr<FDataContainer, ESPMode::ThreadSafe>& DefaultValue) const
+{
+	if (Key ==  UE::CaptureCardMediaSource::OpenColorIOSettings)
+	{
+		FOpenColorIODataContainer Container;
+		Container.ColorConversionSettings = ColorConversionSettings;
+		return MakeShared<FOpenColorIODataContainer>(MoveTemp(Container));
+	}
+
+	return Super::GetMediaOption(Key, DefaultValue);
+}
+
 FString UCaptureCardMediaSource::GetMediaOption(const FName& Key, const FString& DefaultValue) const
 {
 	if (Key == UE::CaptureCardMediaSource::Deinterlacer)
@@ -81,7 +93,8 @@ bool UCaptureCardMediaSource::HasMediaOption(const FName& Key) const
 		|| Key == UE::CaptureCardMediaSource::SourceColorSpace
 		|| Key == UE::CaptureCardMediaSource::RenderJIT
 		|| Key == UE::CaptureCardMediaSource::Framelock
-		|| Key == UE::CaptureCardMediaSource::EvaluationType)
+		|| Key == UE::CaptureCardMediaSource::EvaluationType
+		|| Key == UE::CaptureCardMediaSource::OpenColorIOSettings)
 		{
 			return true;
 		}
diff --git a/Engine/Plugins/Media/MediaIOFramework/Source/MediaIOCore/Private/Player/MediaIOCorePlayerBase.cpp b/Engine/Plugins/Media/MediaIOFramework/Source/MediaIOCore/Private/Player/MediaIOCorePlayerBase.cpp
--- a/Engine/Plugins/Media/MediaIOFramework/Source/MediaIOCore/Private/Player/MediaIOCorePlayerBase.cpp
+++ b/Engine/Plugins/Media/MediaIOFramework/Source/MediaIOCore/Private/Player/MediaIOCorePlayerBase.cpp
@@ -566,6 +566,12 @@ bool FMediaIOCorePlayerBase::ReadMediaOptions(const IMediaOptions* Options)
 	BaseSettings.PlayerTimeOffset = MediaIOCorePlayerDetail::HighPerformanceClockDelta;
 	BaseSettings.AbsoluteMaxBufferSize = MediaIOCorePlayerDetail::CVarMediaIOMaxBufferSize.GetValueOnGameThread();
 
+	const TSharedPtr<UCaptureCardMediaSource::FOpenColorIODataContainer> OCIODataContainer = StaticCastSharedPtr<UCaptureCardMediaSource::FOpenColorIODataContainer>(Options->GetMediaOption(UE::CaptureCardMediaSource::OpenColorIOSettings, TSharedPtr<UCaptureCardMediaSource::FOpenColorIODataContainer, ESPMode::ThreadSafe>()));
+	if (OCIODataContainer)
+	{
+		OCIOSettings = MakeShared<FOpenColorIOColorConversionSettings>(OCIODataContainer->ColorConversionSettings);
+	}
+
 	return true;
 }
 
diff --git a/Engine/Plugins/Media/MediaIOFramework/Source/MediaIOCore/Private/Shared/MediaIOCoreTextureSampleBase.cpp b/Engine/Plugins/Media/MediaIOFramework/Source/MediaIOCore/Private/Shared/MediaIOCoreTextureSampleBase.cpp
--- a/Engine/Plugins/Media/MediaIOFramework/Source/MediaIOCore/Private/Shared/MediaIOCoreTextureSampleBase.cpp
+++ b/Engine/Plugins/Media/MediaIOFramework/Source/MediaIOCore/Private/Shared/MediaIOCoreTextureSampleBase.cpp
@@ -4,8 +4,17 @@
 
 #include "MediaIOCoreTextureSampleConverter.h"
 
+#include "OpenColorIOConfiguration.h"
+#include "OpenColorIOColorSpace.h"
+#include "OpenColorIORendering.h"
+#include "OpenColorIOShader.h"
+#include "OpenColorIOShared.h"
+#include "RenderGraphBuilder.h"
 #include "RHI.h"
 #include "RHIResources.h"
+#include "ScreenPass.h"
+
+DECLARE_GPU_STAT(MediaIO_ColorConversion);
 
 FMediaIOCoreTextureSampleBase::FMediaIOCoreTextureSampleBase()
 	: Duration(FTimespan::Zero())
@@ -196,6 +205,75 @@ void FMediaIOCoreTextureSampleBase::CopyConfiguration(const TSharedPtr<FMediaIOC
 	OriginalSample = SourceSample;
 }
 
+bool FMediaIOCoreTextureSampleBase::ApplyColorConversion(FTexture2DRHIRef& InSrcTexture, FTexture2DRHIRef& InDstTexture)
+{
+	FRHICommandListImmediate& RHICmdList = FRHICommandListExecutor::GetImmediateCommandList();
+
+	if (!CachedOCIOResources.IsValid())
+	{
+		if (ColorConversionSettings.IsValid() && ColorConversionSettings->IsValid())
+		{
+			CachedOCIOResources = MakeShared<FOpenColorIORenderPassResources>();
+
+			FOpenColorIOTransformResource* ShaderResource = nullptr;
+			TSortedMap<int32, FTextureResource*> TransformTextureResources;
+
+			if (ColorConversionSettings->ConfigurationSource != nullptr)
+			{
+				const bool bFoundTransform = ColorConversionSettings->ConfigurationSource->GetRenderResources(
+					GMaxRHIFeatureLevel
+					, *ColorConversionSettings
+					, ShaderResource
+					, TransformTextureResources);
+
+				if (bFoundTransform)
+				{
+					// Transform was found, so shader must be there but doesn't mean the actual shader is available
+					check(ShaderResource);
+					if (ShaderResource->GetShader<FOpenColorIOPixelShader>().IsNull())
+					{
+						ensureMsgf(false, TEXT("Can't apply display look - Shader was invalid for Resource %s"), *ShaderResource->GetFriendlyName());
+
+						//Invalidate shader resource
+						ShaderResource = nullptr;
+					}
+				}
+			}
+
+			CachedOCIOResources->ShaderResource = ShaderResource;
+			CachedOCIOResources->TextureResources = TransformTextureResources;
+		}
+	}
+
+	if (CachedOCIOResources)
+	{
+		FRDGBuilder GraphBuilder(RHICmdList);
+		
+		{
+			RDG_GPU_STAT_SCOPE(GraphBuilder, MediaIO_ColorConversion);
+
+			const FRDGTextureRef ColorConversionInput = GraphBuilder.RegisterExternalTexture(CreateRenderTarget(InSrcTexture, TEXT("MediaTextureResourceColorConverisonInputRT")));
+			const FRDGTextureRef ColorConversionOutput = GraphBuilder.RegisterExternalTexture(CreateRenderTarget(InDstTexture, TEXT("MediaTextureResourceColorConverisonOutputRT")));
+
+			constexpr float DefaultDisplayGamma = 1;
+		
+			FOpenColorIORendering::AddPass_RenderThread(GraphBuilder,
+				FScreenPassViewInfo(),
+				GMaxRHIFeatureLevel,
+				FScreenPassTexture(ColorConversionInput),
+				FScreenPassRenderTarget(ColorConversionOutput, ERenderTargetLoadAction::EClear),
+				*CachedOCIOResources,
+				DefaultDisplayGamma);
+		}
+
+		GraphBuilder.Execute();
+
+		return true;
+	}
+	
+	return false;
+}
+
 #if WITH_ENGINE
 IMediaTextureSampleConverter* FMediaIOCoreTextureSampleBase::GetMediaTextureSampleConverter()
 {
@@ -206,6 +284,15 @@ FRHITexture* FMediaIOCoreTextureSampleBase::GetTexture() const
 {
 	return Texture.GetReference();
 }
+
+IMediaTextureSampleColorConverter* FMediaIOCoreTextureSampleBase::GetMediaTextureSampleColorConverter()
+{
+	if (ColorConversionSettings && ColorConversionSettings->IsValid())
+	{
+		return this;
+	}
+	return nullptr;
+}
 #endif
 
 void FMediaIOCoreTextureSampleBase::SetTexture(TRefCountPtr<FRHITexture> InRHITexture)
@@ -243,6 +330,10 @@ const FMatrix& FMediaIOCoreTextureSampleBase::GetYUVToRGBMatrix() const
 
 bool FMediaIOCoreTextureSampleBase::IsOutputSrgb() const
 {
+	if (ColorConversionSettings && ColorConversionSettings->IsValid())
+	{
+		return false; // Don't apply gamma correction as it will be handled by the OCIO conversion.
+	}
 	return Encoding == UE::Color::EEncoding::sRGB;
 }
 
@@ -264,7 +355,7 @@ FVector2f FMediaIOCoreTextureSampleBase::GetDisplayPrimaryRed() const
 FVector2f FMediaIOCoreTextureSampleBase::GetDisplayPrimaryGreen() const
 {
 	return FVector2f(ColorSpaceStruct.GetGreenChromaticity());
-}
+} 
 
 FVector2f FMediaIOCoreTextureSampleBase::GetDisplayPrimaryBlue() const
 {
@@ -273,10 +364,20 @@ FVector2f FMediaIOCoreTextureSampleBase::GetDisplayPrimaryBlue() const
 
 UE::Color::EEncoding FMediaIOCoreTextureSampleBase::GetEncodingType() const
 {
+	if (ColorConversionSettings && ColorConversionSettings->IsValid())
+    {
+    	return UE::Color::EEncoding::Linear;
+    }
+    
 	return Encoding;
 }
 
 float FMediaIOCoreTextureSampleBase::GetHDRNitsNormalizationFactor() const
 {
+	if (ColorConversionSettings && ColorConversionSettings->IsValid())
+    {
+    	return 1.0f;
+    }
+	
 	return (GetEncodingType() == UE::Color::EEncoding::sRGB || GetEncodingType() == UE::Color::EEncoding::Linear) ? 1.0f : kMediaSample_HDR_NitsNormalizationFactor;
 }
diff --git a/Engine/Plugins/Media/MediaIOFramework/Source/MediaIOCore/Public/CaptureCardMediaSource.h b/Engine/Plugins/Media/MediaIOFramework/Source/MediaIOCore/Public/CaptureCardMediaSource.h
--- a/Engine/Plugins/Media/MediaIOFramework/Source/MediaIOCore/Public/CaptureCardMediaSource.h
+++ b/Engine/Plugins/Media/MediaIOFramework/Source/MediaIOCore/Public/CaptureCardMediaSource.h
@@ -5,14 +5,17 @@
 #include "TimeSynchronizableMediaSource.h"
 #include "Containers/UnrealString.h"
 #include "Engine/TextureDefines.h"
+#include "HAL/Platform.h"
 #include "MediaIOCoreDefinitions.h"
 #include "MediaIOCoreDeinterlacer.h"
-#include "HAL/Platform.h"
+#include "OpenColorIOColorSpace.h"
+#include "OpenColorIORendering.h"
 #include "UObject/NameTypes.h"
 #include "UObject/SoftObjectPath.h"
 #include "UObject/ObjectMacros.h"
 #include "UObject/UObjectGlobals.h"
 
+
 #include "CaptureCardMediaSource.generated.h"
 
 namespace UE::CaptureCardMediaSource
@@ -26,6 +29,7 @@ namespace UE::CaptureCardMediaSource
 	static FName SourceEncoding("SourceEncoding");
 	static FName OverrideSourceColorSpace("OverrideSourceColorSpace");
 	static FName SourceColorSpace("SourceColorcSpace");
+	static FName OpenColorIOSettings("OpenColorIOSettings");
 }
 
 
@@ -94,12 +98,25 @@ public:
 	UPROPERTY(BlueprintReadOnly, EditAnywhere, Category = "Video", meta = (EditCondition = "bOverrideSourceColorSpace"))
 	ETextureColorSpace OverrideSourceColorSpace = ETextureColorSpace::TCS_sRGB;
 
+	/**
+	 * OCIO Settings used for applying a color conversion to the incoming source.
+	 */
+	UPROPERTY(BlueprintReadWrite ,EditAnywhere, Category="Video")
+	FOpenColorIOColorConversionSettings ColorConversionSettings;
+
+public:
+	struct FOpenColorIODataContainer : public IMediaOptions::FDataContainer
+	{
+		FOpenColorIOColorConversionSettings ColorConversionSettings;
+	};
+
 public:
 	//~ IMediaOptions interface
 	using Super::GetMediaOption;
 	virtual FString GetMediaOption(const FName& Key, const FString& DefaultValue) const override;
 	virtual int64 GetMediaOption(const FName& Key, int64 DefaultValue) const override;
 	virtual bool GetMediaOption(const FName& Key, bool bDefaultValue) const override;
+	virtual TSharedPtr<FDataContainer, ESPMode::ThreadSafe> GetMediaOption(const FName& Key, const TSharedPtr<FDataContainer, ESPMode::ThreadSafe>& DefaultValue) const override;
 	virtual bool HasMediaOption(const FName& Key) const override;
 
 public:
diff --git a/Engine/Plugins/Media/MediaIOFramework/Source/MediaIOCore/Public/MediaIOCorePlayerBase.h b/Engine/Plugins/Media/MediaIOFramework/Source/MediaIOCore/Public/MediaIOCorePlayerBase.h
--- a/Engine/Plugins/Media/MediaIOFramework/Source/MediaIOCore/Public/MediaIOCorePlayerBase.h
+++ b/Engine/Plugins/Media/MediaIOFramework/Source/MediaIOCore/Public/MediaIOCorePlayerBase.h
@@ -312,6 +312,9 @@ protected:
 	/** Base set of settings to start from when setuping channels */
 	FMediaIOSamplingSettings BaseSettings;
 
+	/** Open color IO conversion data. */
+	TSharedPtr<struct FOpenColorIOColorConversionSettings> OCIOSettings;
+
 private:
 	void OnSampleDestroyed(TRefCountPtr<FRHITexture> InTexture);
 	void RegisterSampleBuffer(const TSharedPtr<FMediaIOCoreTextureSampleBase>& InSample);
@@ -322,6 +325,7 @@ private:
 	void PreGPUTransferJITR(const TSharedPtr<FMediaIOCoreTextureSampleBase>& InSample, const TSharedPtr<FMediaIOCoreTextureSampleBase>& InJITRProxySample);
 	void ExecuteGPUTransfer(const TSharedPtr<FMediaIOCoreTextureSampleBase>& InSample);
 
+
 	/** GPU Texture transfer object */
 	UE::GPUTextureTransfer::TextureTransferPtr GPUTextureTransfer;
 
diff --git a/Engine/Plugins/Media/MediaIOFramework/Source/MediaIOCore/Public/MediaIOCoreTextureSampleBase.h b/Engine/Plugins/Media/MediaIOFramework/Source/MediaIOCore/Public/MediaIOCoreTextureSampleBase.h
--- a/Engine/Plugins/Media/MediaIOFramework/Source/MediaIOCore/Public/MediaIOCoreTextureSampleBase.h
+++ b/Engine/Plugins/Media/MediaIOFramework/Source/MediaIOCore/Public/MediaIOCoreTextureSampleBase.h
@@ -3,6 +3,8 @@
 #pragma once
 
 #include "IMediaTextureSample.h"
+#include "MediaIOCoreDefinitions.h"
+#include "IMediaTextureSampleConverter.h"
 #include "MediaObjectPool.h"
 
 
@@ -84,6 +86,7 @@ struct MEDIAIOCORE_API FMediaIOCoreSampleJITRConfigurationArgs
 class MEDIAIOCORE_API FMediaIOCoreTextureSampleBase
 	: public IMediaTextureSample
 	, public IMediaPoolable
+	, public IMediaTextureSampleColorConverter
 {
 
 public:
@@ -205,6 +208,13 @@ public:
 	 */
 	bool SetBufferWithEvenOddLine(bool bUseEvenLine, const void* InVideoBuffer, uint32 InBufferSize, uint32 InStride, uint32 InHeight);
 
+	/**
+	 * Set the OCIO settings used for color conversion.
+	 */
+	void SetColorConversionSettings(TSharedPtr<struct FOpenColorIOColorConversionSettings> InColorConversionSettings)
+	{
+		ColorConversionSettings = InColorConversionSettings;
+	}
 	/**
 	 * Request an uninitialized sample buffer.
 	 * Should be used when the buffer could be filled by something else.
@@ -327,9 +337,11 @@ public:
 		}
 
 		return Buffer.GetData();
-
 	}
 
+	//~ IMediaTextureSampleColorConverter interface
+	virtual bool ApplyColorConversion(FTexture2DRHIRef& InSrcTexture, FTexture2DRHIRef& InDstTexture) override;
+
 	void* GetMutableBuffer()
 	{
 		if (ExternalBuffer)
@@ -343,6 +355,8 @@ public:
 #if WITH_ENGINE
 	virtual IMediaTextureSampleConverter* GetMediaTextureSampleConverter() override;
 	virtual FRHITexture* GetTexture() const override;
+
+	virtual IMediaTextureSampleColorConverter* GetMediaTextureSampleColorConverter() override;
 #endif //WITH_ENGINE
 
 	void SetBuffer(void* InBuffer)
@@ -435,4 +449,10 @@ private:
 
 	/** Time offset evaluated on game thread for JITR */
 	double EvaluationOffsetInSeconds = 0;
+
+private:
+	/** Settings used to apply an OCIO conversion to the sample. */
+	TSharedPtr<struct FOpenColorIOColorConversionSettings> ColorConversionSettings;
+	/** Cached rendering resources used for the color conversion pass when using OCIO. */
+	TSharedPtr<struct FOpenColorIORenderPassResources> CachedOCIOResources;
 };
