[MediaIO] Added missing OCIO data copy for JITR

diff --git a/Engine/Plugins/Media/MediaIOFramework/Source/MediaIOCore/Private/Shared/MediaIOCoreTextureSampleBase.cpp b/Engine/Plugins/Media/MediaIOFramework/Source/MediaIOCore/Private/Shared/MediaIOCoreTextureSampleBase.cpp
--- a/Engine/Plugins/Media/MediaIOFramework/Source/MediaIOCore/Private/Shared/MediaIOCoreTextureSampleBase.cpp
+++ b/Engine/Plugins/Media/MediaIOFramework/Source/MediaIOCore/Private/Shared/MediaIOCoreTextureSampleBase.cpp
@@ -200,6 +200,8 @@ void FMediaIOCoreTextureSampleBase::CopyConfiguration(const TSharedPtr<FMediaIOC
 	Encoding = SourceSample->Encoding;
 	ColorSpace = SourceSample->ColorSpace;
 	ColorSpaceStruct = SourceSample->ColorSpaceStruct;
+	ColorConversionSettings = SourceSample->ColorConversionSettings;
+	CachedOCIOResources = SourceSample->CachedOCIOResources;
 
 	// Save original sample
 	OriginalSample = SourceSample;
