Fix audio not being output when starting a media capture after starting PIE

diff --git a/Engine/Plugins/Media/MediaIOFramework/Source/MediaIOCore/Private/MediaCapture.cpp b/Engine/Plugins/Media/MediaIOFramework/Source/MediaIOCore/Private/MediaCapture.cpp
--- a/Engine/Plugins/Media/MediaIOFramework/Source/MediaIOCore/Private/MediaCapture.cpp
+++ b/Engine/Plugins/Media/MediaIOFramework/Source/MediaIOCore/Private/MediaCapture.cpp
@@ -541,6 +541,8 @@ void UMediaCapture::StopCapture(bool bAllowPendingFrameToBeProcess)
 			
 			StopCaptureImpl(bAllowPendingFrameToBeProcess);
 
+			SetCaptureAudioDevice(FAudioDeviceHandle());
+
 			if (CaptureSource)
 			{
 				constexpr bool bFlushRenderingCommands = false;
diff --git a/Engine/Plugins/Media/MediaIOFramework/Source/MediaIOCore/Private/MediaIOCoreAudioOutput.cpp b/Engine/Plugins/Media/MediaIOFramework/Source/MediaIOCore/Private/MediaIOCoreAudioOutput.cpp
--- a/Engine/Plugins/Media/MediaIOFramework/Source/MediaIOCore/Private/MediaIOCoreAudioOutput.cpp
+++ b/Engine/Plugins/Media/MediaIOFramework/Source/MediaIOCore/Private/MediaIOCoreAudioOutput.cpp
@@ -16,6 +16,36 @@
 
 DEFINE_LOG_CATEGORY_STATIC(LogMediaIOAudioOutput, Log, All);
 
+namespace MediaIOCoreAudioOutputUtils
+{
+	UWorld* GetCurrentWorld()
+	{
+#if WITH_EDITOR
+		if (GEditor)
+		{
+			FWorldContext* PIEWorldContext = GEditor->GetPIEWorldContext();
+			if (PIEWorldContext)
+			{
+				return PIEWorldContext->World();
+			}
+			return GEditor->GetEditorWorldContext(false).World();
+		}
+		if (GWorld)
+		{
+			return GWorld->GetWorld();
+		}
+#endif
+		for (const FWorldContext& WorldContext : GEngine->GetWorldContexts())
+		{
+			if (WorldContext.WorldType == EWorldType::Game)
+			{
+				return WorldContext.World();
+			}
+		}
+		return nullptr;
+	}
+}
+
 #define LOCTEXT_NAMESPACE "MediaIOAudioOutput"
 
 FMediaIOAudioOutput::FMediaIOAudioOutput(Audio::FPatchOutputStrongPtr InPatchOutput, const FAudioOptions& InAudioOptions)
@@ -159,7 +189,7 @@ void FMediaIOAudioCapture::RegisterAudioDevice(const FAudioDeviceHandle& InAudio
 
 void FMediaIOAudioCapture::UnregisterAudioDevice()
 {
-	if (RegisteredDeviceId != INDEX_NONE)
+	if (RegisteredDeviceId != INDEX_NONE && GEngine && GEngine->GetAudioDeviceManager())
 	{
 		if (FAudioDevice* RegisteredDevice = GEngine->GetAudioDeviceManager()->GetAudioDeviceRaw(RegisteredDeviceId))
 		{
@@ -172,7 +202,9 @@ void FMediaIOAudioCapture::UnregisterAudioDevice()
 FMainMediaIOAudioCapture::FMainMediaIOAudioCapture()
 	: FMediaIOAudioCapture(FAudioDeviceHandle())
 {
-	RegisterMainAudioDevice();
+	// Register the current device (PIE or main).
+	RegisterCurrentAudioDevice();
+	
 #if WITH_EDITOR
 	FEditorDelegates::PostPIEStarted.AddRaw(this, &FMainMediaIOAudioCapture::OnPIEStarted);
 	FEditorDelegates::PrePIEEnded.AddRaw(this, &FMainMediaIOAudioCapture::OnPIEEnded);
@@ -182,50 +214,39 @@ FMainMediaIOAudioCapture::FMainMediaIOAudioCapture()
 FMainMediaIOAudioCapture::~FMainMediaIOAudioCapture()
 {
 #if WITH_EDITOR
-	FEditorDelegates::EndPIE.RemoveAll(this);
+	FEditorDelegates::PrePIEEnded.RemoveAll(this);
 	FEditorDelegates::PostPIEStarted.RemoveAll(this);
 #endif
-
-	UnregisterMainAudioDevice();
 }
 
 #if WITH_EDITOR
 void FMainMediaIOAudioCapture::OnPIEStarted(const bool)
 {
-	UnregisterMainAudioDevice();
-	
-	if (GEditor)
-	{
-		if (FWorldContext* PIEWorldContext = GEditor->GetPIEWorldContext())
-		{
-			Audio::FMixerDevice* MixerDevice = static_cast<Audio::FMixerDevice*>(PIEWorldContext->World()->GetAudioDeviceRaw());
-			RegisterBufferListener(MixerDevice);
-		}
-	}
+	RegisterCurrentAudioDevice();
 }
 
 void FMainMediaIOAudioCapture::OnPIEEnded(const bool)
 {
-	if (GEditor)
-	{
-		if (FWorldContext* PIEWorldContext = GEditor->GetPIEWorldContext())
-		{
-			UnregisterBufferListener(PIEWorldContext->World()->GetAudioDeviceRaw());
-		}
-	}
-	
+	// Note: PIE context still active, need to explicitly fallback to engine's device.
 	RegisterMainAudioDevice();
 }
 #endif
 
 void FMainMediaIOAudioCapture::RegisterMainAudioDevice()
 {
-	RegisterBufferListener(GEngine->GetMainAudioDeviceRaw());
+	RegisterAudioDevice(GEngine->GetMainAudioDevice());
 }
 
-void FMainMediaIOAudioCapture::UnregisterMainAudioDevice()
+void FMainMediaIOAudioCapture::RegisterCurrentAudioDevice()
 {
-	UnregisterBufferListener(GEngine->GetMainAudioDeviceRaw());
+	if (const UWorld* World = MediaIOCoreAudioOutputUtils::GetCurrentWorld())
+	{
+		RegisterAudioDevice(World->GetAudioDevice());
+	}
+	else
+	{
+		RegisterMainAudioDevice();
+	}
 }
 
 #undef LOCTEXT_NAMESPACE //MediaIOAudioOutput
\ No newline at end of file
diff --git a/Engine/Plugins/Media/MediaIOFramework/Source/MediaIOCore/Public/MediaIOCoreAudioOutput.h b/Engine/Plugins/Media/MediaIOFramework/Source/MediaIOCore/Public/MediaIOCoreAudioOutput.h
--- a/Engine/Plugins/Media/MediaIOFramework/Source/MediaIOCore/Public/MediaIOCoreAudioOutput.h
+++ b/Engine/Plugins/Media/MediaIOFramework/Source/MediaIOCore/Public/MediaIOCoreAudioOutput.h
@@ -181,7 +181,6 @@ protected:
 	void RegisterBufferListener(FAudioDevice* AudioDevice);
 	void UnregisterBufferListener(FAudioDevice* AudioDevice);
 
-private:
 	void RegisterAudioDevice(const FAudioDeviceHandle& InAudioDeviceHandle);
 	void UnregisterAudioDevice();
 
@@ -215,7 +214,7 @@ class FMainMediaIOAudioCapture : public FMediaIOAudioCapture
 {
 public:
 	FMainMediaIOAudioCapture();
-	virtual ~FMainMediaIOAudioCapture();
+	virtual ~FMainMediaIOAudioCapture() override;
 
 private:
 #if WITH_EDITOR
@@ -224,5 +223,5 @@ private:
 #endif
 
 	void RegisterMainAudioDevice();
-	void UnregisterMainAudioDevice();
+	void RegisterCurrentAudioDevice();
 };
\ No newline at end of file
