Fix crash when doing a media capture with resample pass in standalone game

diff --git a/Engine/Plugins/Media/MediaIOFramework/Source/MediaIOCore/Private/MediaCaptureRenderPass.cpp b/Engine/Plugins/Media/MediaIOFramework/Source/MediaIOCore/Private/MediaCaptureRenderPass.cpp
--- a/Engine/Plugins/Media/MediaIOFramework/Source/MediaIOCore/Private/MediaCaptureRenderPass.cpp
+++ b/Engine/Plugins/Media/MediaIOFramework/Source/MediaIOCore/Private/MediaCaptureRenderPass.cpp
@@ -67,8 +67,11 @@ namespace UE::MediaCapture::Resample
 	void AddResamplePass(FRDGBuilder& GraphBuilder, const FViewInfo& View, const FScreenPassTexture& InputTexture, FRDGTextureRef OutputTexture)
 	{
 		check(InputTexture.Texture);
-
-		const FScreenPassRenderTarget Input = FScreenPassRenderTarget::CreateFromInput(GraphBuilder, InputTexture, ERenderTargetLoadAction::ENoAction, TEXT("MediaCaptureResampleInput"));
+		FScreenPassRenderTarget Input;
+		const FIntVector InputSize = InputTexture.Texture->Desc.GetSize();
+		Input.Texture = InputTexture.Texture;
+		Input.ViewRect = FIntRect{ 0, 0, InputSize.X, InputSize.Y };
+		Input.LoadAction = ERenderTargetLoadAction::ENoAction;
 
 		FScreenPassRenderTarget Output;
 		const FIntVector OutputSize = OutputTexture->Desc.GetSize();
@@ -167,6 +170,7 @@ namespace UE::MediaCapture
 		ResampleOutputTextureDesc.Reset();
 		ResampleOutputTextureDesc.Extent.X = Args.MediaCapture->GetDesiredSize().X;
 		ResampleOutputTextureDesc.Extent.Y = Args.MediaCapture->GetDesiredSize().Y;
+		ResampleOutputTextureDesc.Flags = TexCreate_None;
 		ResampleOutputTextureDesc.Flags |= TexCreate_RenderTargetable | TexCreate_UAV | TexCreate_NoFastClear;
 		ResampleOutputTextureDesc.ClearValue = FClearValueBinding(FLinearColor(0, 0, 0, 0));
 
