Fix media capture options losing previous ResizeSourceBuffer option

diff --git a/Engine/Plugins/Media/MediaIOFramework/Source/MediaIOCore/Private/MediaCapture.cpp b/Engine/Plugins/Media/MediaIOFramework/Source/MediaIOCore/Private/MediaCapture.cpp
--- a/Engine/Plugins/Media/MediaIOFramework/Source/MediaIOCore/Private/MediaCapture.cpp
+++ b/Engine/Plugins/Media/MediaIOFramework/Source/MediaIOCore/Private/MediaCapture.cpp
@@ -42,6 +42,7 @@
 #include "Framework/Notifications/NotificationManager.h"
 #include "IAssetViewport.h"
 #include "LevelEditor.h"
+#include "UObject/UE5ReleaseStreamObjectVersion.h"
 #include "Widgets/Notifications/SNotificationList.h"
 #endif //WITH_EDITOR
 
@@ -113,6 +114,21 @@ FMediaCaptureOptions::FMediaCaptureOptions()
 
 }
 
+void FMediaCaptureOptions::PostSerialize(const FArchive& Ar)
+{
+#if WITH_EDITOR
+	if (Ar.IsLoading())
+	{
+		if (Ar.CustomVer(FUE5ReleaseStreamObjectVersion::GUID) < FUE5ReleaseStreamObjectVersion::MediaCaptureNewResizeMethods)
+		{
+			PRAGMA_DISABLE_DEPRECATION_WARNINGS
+			ResizeMethod = bResizeSourceBuffer_DEPRECATED ? EMediaCaptureResizeMethod::ResizeSource : EMediaCaptureResizeMethod::ResizeInRenderPass;
+			PRAGMA_ENABLE_DEPRECATION_WARNINGS
+		}
+	}
+#endif
+}
+
 
 /* UMediaCapture
 *****************************************************************************/
diff --git a/Engine/Plugins/Media/MediaIOFramework/Source/MediaIOCore/Public/MediaCapture.h b/Engine/Plugins/Media/MediaIOFramework/Source/MediaIOCore/Public/MediaCapture.h
--- a/Engine/Plugins/Media/MediaIOFramework/Source/MediaIOCore/Public/MediaCapture.h
+++ b/Engine/Plugins/Media/MediaIOFramework/Source/MediaIOCore/Public/MediaCapture.h
@@ -168,6 +168,7 @@ struct MEDIAIOCORE_API FMediaCaptureOptions
 	FMediaCaptureOptions();
 
 public:
+	void PostSerialize(const FArchive& Ar);
 
 	/** If source texture's size change, auto restart capture to follow source resolution if implementation supports it. */
 	UPROPERTY(BlueprintReadWrite, EditAnywhere, Category = "MediaCapture")
@@ -734,3 +735,11 @@ private:
 	friend class UE::MediaCaptureData::FBufferCaptureFrame;
 };
 
+
+template<> struct TStructOpsTypeTraits<FMediaCaptureOptions> : public TStructOpsTypeTraitsBase2<FMediaCaptureOptions>
+{
+	enum
+	{
+		WithPostSerialize = true
+	};
+};
diff --git a/Engine/Source/Runtime/Core/Public/UObject/UE5ReleaseStreamObjectVersion.h b/Engine/Source/Runtime/Core/Public/UObject/UE5ReleaseStreamObjectVersion.h
--- a/Engine/Source/Runtime/Core/Public/UObject/UE5ReleaseStreamObjectVersion.h
+++ b/Engine/Source/Runtime/Core/Public/UObject/UE5ReleaseStreamObjectVersion.h
@@ -139,6 +139,9 @@ struct FUE5ReleaseStreamObjectVersion
 		// Changed default tangent behavior for new dynamic mesh components
 		DynamicMeshComponentsDefaultUseExternalTangents,
 		
+		// Added resize methods to media capture
+		MediaCaptureNewResizeMethods,
+		
 		// -----<new versions can be added above this line>-------------------------------------------------
 		VersionPlusOne,
 		LatestVersion = VersionPlusOne - 1
