Fix crash when doing a media capture with capture phase BeforePostProcessing

diff --git a/Engine/Plugins/Media/MediaIOFramework/Source/MediaIOCore/Private/MediaCaptureSceneViewExtension.h b/Engine/Plugins/Media/MediaIOFramework/Source/MediaIOCore/Private/MediaCaptureSceneViewExtension.h
--- a/Engine/Plugins/Media/MediaIOFramework/Source/MediaIOCore/Private/MediaCaptureSceneViewExtension.h
+++ b/Engine/Plugins/Media/MediaIOFramework/Source/MediaIOCore/Private/MediaCaptureSceneViewExtension.h
@@ -95,7 +95,7 @@ public:
 		TRACE_CPUPROFILER_EVENT_SCOPE(MediaCaptureExtensionCallback);
 
 		FScreenPassTexture SceneColor = InOutInputs.GetInput(EPostProcessMaterialInput::SceneColor);
-		InOutInputs.Validate();
+		check(SceneColor.IsValid());
 
 		if (FRDGTextureRef TextureRef = SceneColor.Texture)
 		{
