Fix color flashes when changing OCIO config on blackmagic media source

diff --git a/Engine/Plugins/Media/MediaIOFramework/Source/MediaIOCore/Public/MediaIOCoreTextureSampleBase.h b/Engine/Plugins/Media/MediaIOFramework/Source/MediaIOCore/Public/MediaIOCoreTextureSampleBase.h
--- a/Engine/Plugins/Media/MediaIOFramework/Source/MediaIOCore/Public/MediaIOCoreTextureSampleBase.h
+++ b/Engine/Plugins/Media/MediaIOFramework/Source/MediaIOCore/Public/MediaIOCoreTextureSampleBase.h
@@ -384,6 +384,7 @@ protected:
 	{
 		Buffer.Reset();
 		ExternalBuffer = nullptr;
+		CachedOCIOResources.Reset();
 		Texture.SafeRelease();
 	}
 
