Fix crash when doing a crop with file media output with ResizeRenderPass in PIE

diff --git a/Engine/Plugins/Media/MediaIOFramework/Source/MediaIOCore/Private/MediaCaptureRenderPass.cpp b/Engine/Plugins/Media/MediaIOFramework/Source/MediaIOCore/Private/MediaCaptureRenderPass.cpp
--- a/Engine/Plugins/Media/MediaIOFramework/Source/MediaIOCore/Private/MediaCaptureRenderPass.cpp
+++ b/Engine/Plugins/Media/MediaIOFramework/Source/MediaIOCore/Private/MediaCaptureRenderPass.cpp
@@ -374,7 +374,14 @@ namespace UE::MediaCapture
 
 		if (MediaCapture->GetDesiredCaptureOptions().ResizeMethod == EMediaCaptureResizeMethod::ResizeInRenderPass)
 		{
-			RenderPasses.Add(CreateResamplePass());
+			if (MediaCapture->GetDesiredCaptureOptions().Crop != EMediaCaptureCroppingType::None)
+			{
+				UE_LOG(LogMediaIOCore, Warning, TEXT("The capture for %s will not be cropped because ResizeMethod was already set to \'Resize in Render Pass\' and these options are mutually exclusive."), *MediaCapture->GetMediaOutputName());
+			}
+			else
+			{
+				RenderPasses.Add(CreateResamplePass());
+			}
 		}
 		
 		if (MediaCapture->GetDesiredCaptureOptions().ColorConversionSettings.IsValid())
