Media Capture Options
- Only allow resizing if Crop is set to none
- Only allow modifying custom capture options if crop is set to custom

diff --git a/Engine/Plugins/Media/MediaIOFramework/Source/MediaIOCore/Public/MediaCapture.h b/Engine/Plugins/Media/MediaIOFramework/Source/MediaIOCore/Public/MediaCapture.h
--- a/Engine/Plugins/Media/MediaIOFramework/Source/MediaIOCore/Public/MediaCapture.h
+++ b/Engine/Plugins/Media/MediaIOFramework/Source/MediaIOCore/Public/MediaCapture.h
@@ -190,14 +190,14 @@ public:
 	 * Crop the captured SceneViewport or TextureRenderTarget2D to the desired size.
 	 * @note Only valid when Crop is set to Custom.
 	 */
-	UPROPERTY(BlueprintReadWrite, EditAnywhere, Category="MediaCapture")
+	UPROPERTY(BlueprintReadWrite, EditAnywhere, Category="MediaCapture", meta=(editcondition = "Crop == EMediaCaptureCroppingType::Custom"))
 	FIntPoint CustomCapturePoint;
 	
 	/**
 	 * When the capture start, control if and how the source buffer should be resized to the desired size.
 	 * @note Only valid when a size is specified by the MediaOutput.
 	 */
-	UPROPERTY(BlueprintReadWrite, EditAnywhere, Category="MediaCapture")
+	UPROPERTY(BlueprintReadWrite, EditAnywhere, Category="MediaCapture", meta = (editcondition = "Crop == EMediaCaptureCroppingType::None"))
 	EMediaCaptureResizeMethod ResizeMethod = EMediaCaptureResizeMethod::None;
 
 	/**
