- Fixing stack size issue for media capture thread

diff --git a/Engine/Plugins/Media/MediaIOFramework/Source/MediaIOCore/Private/MediaCaptureSyncPointWatcher.cpp b/Engine/Plugins/Media/MediaIOFramework/Source/MediaIOCore/Private/MediaCaptureSyncPointWatcher.cpp
--- a/Engine/Plugins/Media/MediaIOFramework/Source/MediaIOCore/Private/MediaCaptureSyncPointWatcher.cpp
+++ b/Engine/Plugins/Media/MediaIOFramework/Source/MediaIOCore/Private/MediaCaptureSyncPointWatcher.cpp
@@ -15,7 +15,7 @@ FSyncPointWatcher::FSyncPointWatcher(UMediaCapture* InMediaCapture)
 	bIsEnabled = true;
 
 	const FString ThreadName = FString::Printf(TEXT("%s_SyncThread"), *MediaCapture->GetName());
-	WorkingThread.Reset(FRunnableThread::Create(this, *ThreadName, 32 * 1024, TPri_Highest));
+	WorkingThread.Reset(FRunnableThread::Create(this, *ThreadName, 256 * 1024, TPri_Highest));
 }
 
 uint32 FSyncPointWatcher::Run()
