Fix MediaIO audio/captions/subtitles in JITR mode

diff --git a/Engine/Plugins/Media/MediaIOFramework/Source/MediaIOCore/Private/Player/MediaIOCorePlayerBase.cpp b/Engine/Plugins/Media/MediaIOFramework/Source/MediaIOCore/Private/Player/MediaIOCorePlayerBase.cpp
--- a/Engine/Plugins/Media/MediaIOFramework/Source/MediaIOCore/Private/Player/MediaIOCorePlayerBase.cpp
+++ b/Engine/Plugins/Media/MediaIOFramework/Source/MediaIOCore/Private/Player/MediaIOCorePlayerBase.cpp
@@ -919,6 +919,15 @@ void FMediaIOCorePlayerBase::ExecuteGPUTransfer(const TSharedPtr<FMediaIOCoreTex
 	});
 }
 
+FMediaIOCoreSamples& FMediaIOCorePlayerBase::GetSamples_Internal()
+{
+	if (IsJustInTimeRenderingEnabled())
+	{
+		return *JITRSamples;
+	}
+	return *Samples;
+}
+
 TArray<TSharedRef<FMediaIOCoreTextureSampleBase>> FMediaIOCorePlayerBase::Deinterlace(const UE::MediaIOCore::FVideoFrame& InVideoFrame) const
 {
 	if (Deinterlacer)
@@ -1074,22 +1083,22 @@ bool FMediaIOCorePlayerBase::SupportsSubFrames() const
 
 void FMediaIOCorePlayerBase::AddAudioSample(const TSharedRef<FMediaIOCoreAudioSampleBase>& Sample)
 {
-	Samples->AddAudio(Sample);
+	GetSamples_Internal().AddAudio(Sample);
 }
 
 void FMediaIOCorePlayerBase::AddCaptionSample(const TSharedRef<FMediaIOCoreCaptionSampleBase>& Sample)
 {
-	Samples->AddCaption(Sample);
+	GetSamples_Internal().AddCaption(Sample);
 }
 
 void FMediaIOCorePlayerBase::AddMetadataSample(const TSharedRef<FMediaIOCoreBinarySampleBase>& Sample)
 {
-	Samples->AddMetadata(Sample);
+	GetSamples_Internal().AddMetadata(Sample);
 }
 
 void FMediaIOCorePlayerBase::AddSubtitleSample(const TSharedRef<FMediaIOCoreSubtitleSampleBase>& Sample)
 {
-	Samples->AddSubtitle(Sample);
+	GetSamples_Internal().AddSubtitle(Sample);
 }
 
 void FMediaIOCorePlayerBase::AddVideoSample(const TSharedRef<FMediaIOCoreTextureSampleBase>& Sample)
diff --git a/Engine/Plugins/Media/MediaIOFramework/Source/MediaIOCore/Public/MediaIOCorePlayerBase.h b/Engine/Plugins/Media/MediaIOFramework/Source/MediaIOCore/Public/MediaIOCorePlayerBase.h
--- a/Engine/Plugins/Media/MediaIOFramework/Source/MediaIOCore/Public/MediaIOCorePlayerBase.h
+++ b/Engine/Plugins/Media/MediaIOFramework/Source/MediaIOCore/Public/MediaIOCorePlayerBase.h
@@ -12,8 +12,10 @@
 #include "ITimedDataInput.h"
 
 #include "MediaIOCoreDefinitions.h"
+#include "MediaIOCoreSamples.h"
 #include "MediaIOCoreSampleContainer.h"
 
+
 #include "HAL/CriticalSection.h"
 #include "Misc/CoreMisc.h"
 #include "Misc/FrameRate.h"
@@ -325,6 +327,8 @@ private:
 	void PreGPUTransferJITR(const TSharedPtr<FMediaIOCoreTextureSampleBase>& InSample, const TSharedPtr<FMediaIOCoreTextureSampleBase>& InJITRProxySample);
 	void ExecuteGPUTransfer(const TSharedPtr<FMediaIOCoreTextureSampleBase>& InSample);
 
+	/** Get the right samples buffer for pushing samples depending on if we're in JITR mode or not. */
+	FMediaIOCoreSamples& GetSamples_Internal();
 
 	/** GPU Texture transfer object */
 	UE::GPUTextureTransfer::TextureTransferPtr GPUTextureTransfer;
@@ -351,15 +355,14 @@ private:
 
 private:
 
-	/** Special JITR sample container, holding a single sample used to trigger media framework rendering */
+	/** Special JITR sample container, overrides FetchVideo to return a ProxySample that will be populated by JustInTimeSampleRender_RenderThread. */
 	class FJITRMediaTextureSamples
-		: public IMediaSamples
+		: public FMediaIOCoreSamples
 	{
 	public:
 		FJITRMediaTextureSamples() = default;
 		FJITRMediaTextureSamples(const FJITRMediaTextureSamples&) = delete;
 		FJITRMediaTextureSamples& operator=(const FJITRMediaTextureSamples&) = delete;
-
 	public:
 		//~ Begin IMediaSamples interface
 		virtual bool FetchVideo(TRange<FTimespan> TimeRange, TSharedPtr<IMediaTextureSample>& OutSample) override;
