- Removing private rendering folder reference in media io

diff --git a/Engine/Plugins/Media/MediaIOFramework/Source/MediaIOCore/MediaIOCore.Build.cs b/Engine/Plugins/Media/MediaIOFramework/Source/MediaIOCore/MediaIOCore.Build.cs
--- a/Engine/Plugins/Media/MediaIOFramework/Source/MediaIOCore/MediaIOCore.Build.cs
+++ b/Engine/Plugins/Media/MediaIOFramework/Source/MediaIOCore/MediaIOCore.Build.cs
@@ -31,13 +31,6 @@ namespace UnrealBuildTool.Rules
 					"TimeManagement"
 				});
 
-			PrivateIncludePaths.AddRange(
-				new string[] {
-				// required for scene view extensions related headers
-				Path.Combine(GetModuleDirectory("Renderer"), "Private")
-			});
-
-
 			PrivateDependencyModuleNames.AddRange(
                 new string[]
                 {
diff --git a/Engine/Plugins/Media/MediaIOFramework/Source/MediaIOCore/Private/MediaCaptureRenderPass.cpp b/Engine/Plugins/Media/MediaIOFramework/Source/MediaIOCore/Private/MediaCaptureRenderPass.cpp
--- a/Engine/Plugins/Media/MediaIOFramework/Source/MediaIOCore/Private/MediaCaptureRenderPass.cpp
+++ b/Engine/Plugins/Media/MediaIOFramework/Source/MediaIOCore/Private/MediaCaptureRenderPass.cpp
@@ -13,7 +13,6 @@
 #include "PixelShaderUtils.h"
 #include "RenderGraphEvent.h"
 #include "RenderGraphResources.h"
-#include "SceneRendering.h"
 #include "SceneView.h"
 #include "ScreenPass.h"
 #include "ShaderParameterMacros.h"
@@ -64,7 +63,7 @@ namespace UE::MediaCapture::Resample
 	IMPLEMENT_GLOBAL_SHADER(FMediaCaptureResamplePS, "/MediaIOShaders/MediaIO.usf", "MediaDownsample", SF_Pixel);
 	
 	/** Adds a Resample pass to the graph builder. */
-	void AddResamplePass(FRDGBuilder& GraphBuilder, const FViewInfo& View, const FScreenPassTexture& InputTexture, FRDGTextureRef OutputTexture)
+	void AddResamplePass(FRDGBuilder& GraphBuilder, const FScreenPassTexture& InputTexture, FRDGTextureRef OutputTexture)
 	{
 		check(InputTexture.Texture);
 		FScreenPassRenderTarget Input;
@@ -483,20 +482,8 @@ namespace UE::MediaCapture
 			
 		const FIntRect ViewRect(ConversionPassArgs.CopyInfo.GetSourceRect());
 
-		//Dummy ViewFamily/ViewInfo created to use built in Draw Screen/Texture Pass
-		FSceneViewFamily ViewFamily(FSceneViewFamily::ConstructionValues(nullptr, nullptr, FEngineShowFlags(ESFIM_Game))
-			.SetTime(FGameTime()));
-		FSceneViewInitOptions ViewInitOptions;
-		ViewInitOptions.ViewFamily = &ViewFamily;
-		ViewInitOptions.SetViewRectangle(ViewRect);
-		ViewInitOptions.ViewOrigin = FVector::ZeroVector;
-		ViewInitOptions.ViewRotationMatrix = FMatrix::Identity;
-		ViewInitOptions.ProjectionMatrix = FMatrix::Identity;
-			
-		FViewInfo ViewInfo = FViewInfo(ViewInitOptions);
-
 		RDG_GPU_STAT_SCOPE(Args.GraphBuilder, MediaCapture_Resample);
-		MediaCapture::Resample::AddResamplePass(Args.GraphBuilder, ViewInfo, FScreenPassTexture(ConversionPassArgs.SourceRGBTexture), (FRDGTextureRef)OutputTexture);
+		MediaCapture::Resample::AddResamplePass(Args.GraphBuilder, FScreenPassTexture(ConversionPassArgs.SourceRGBTexture), (FRDGTextureRef)OutputTexture);
 	}
 
 	FRenderPipeline::FRenderPipeline(UMediaCapture* InMediaCapture)
diff --git a/Engine/Plugins/Media/MediaIOFramework/Source/MediaIOCore/Private/MediaCaptureSceneViewExtension.h b/Engine/Plugins/Media/MediaIOFramework/Source/MediaIOCore/Private/MediaCaptureSceneViewExtension.h
--- a/Engine/Plugins/Media/MediaIOFramework/Source/MediaIOCore/Private/MediaCaptureSceneViewExtension.h
+++ b/Engine/Plugins/Media/MediaIOFramework/Source/MediaIOCore/Private/MediaCaptureSceneViewExtension.h
@@ -6,14 +6,14 @@
 #include "MediaCapture.h"
 #include "MediaIOCoreModule.h"
 #include "ImagePixelData.h"
-#include "PostProcess/PostProcessing.h"
+#include "PostProcess/PostProcessMaterialInputs.h"
 #include "RenderGraphUtils.h"
 #include "RHI.h"
 #include "RHICommandList.h"
 #include "RHIResources.h"
 #include "SceneView.h"
 #include "ScreenPass.h"
-#include "PostProcess/PostProcessMaterialInputs.h"
+
 
 class FRDGTexture;
 
