MediaIO: Increased default number of buffers from 3 to 4 to avoid starvation.

Starvation may happen when anythread scheduling is disabled.

diff --git a/Engine/Plugins/Media/MediaIOFramework/Source/MediaIOCore/Private/MediaOutput.cpp b/Engine/Plugins/Media/MediaIOFramework/Source/MediaIOCore/Private/MediaOutput.cpp
--- a/Engine/Plugins/Media/MediaIOFramework/Source/MediaIOCore/Private/MediaOutput.cpp
+++ b/Engine/Plugins/Media/MediaIOFramework/Source/MediaIOCore/Private/MediaOutput.cpp
@@ -14,7 +14,7 @@ const FIntPoint UMediaOutput::RequestCaptureSourceSize = FIntPoint::ZeroValue;
 
 UMediaOutput::UMediaOutput(const FObjectInitializer& ObjectInitializer)
 	: Super(ObjectInitializer)
-	, NumberOfTextureBuffers(3)
+	, NumberOfTextureBuffers(4)
 {
 }
 
