- Fixing memory leak when JITR is enabled in MediaIO based players

diff --git a/Engine/Plugins/Media/MediaIOFramework/Source/MediaIOCore/Public/MediaIOCoreTextureSampleBase.h b/Engine/Plugins/Media/MediaIOFramework/Source/MediaIOCore/Public/MediaIOCoreTextureSampleBase.h
--- a/Engine/Plugins/Media/MediaIOFramework/Source/MediaIOCore/Public/MediaIOCoreTextureSampleBase.h
+++ b/Engine/Plugins/Media/MediaIOFramework/Source/MediaIOCore/Public/MediaIOCoreTextureSampleBase.h
@@ -386,6 +386,7 @@ protected:
 		ExternalBuffer = nullptr;
 		CachedOCIOResources.Reset();
 		Texture.SafeRelease();
+		OriginalSample.Reset();
 	}
 
 	/**
