- Fixing crash whenever you enable media capture with audio enabled

diff --git a/Engine/Plugins/Media/MediaIOFramework/Source/MediaIOCore/Private/MediaIOCoreAudioOutput.cpp b/Engine/Plugins/Media/MediaIOFramework/Source/MediaIOCore/Private/MediaIOCoreAudioOutput.cpp
--- a/Engine/Plugins/Media/MediaIOFramework/Source/MediaIOCore/Private/MediaIOCoreAudioOutput.cpp
+++ b/Engine/Plugins/Media/MediaIOFramework/Source/MediaIOCore/Private/MediaIOCoreAudioOutput.cpp
@@ -85,16 +85,22 @@ Audio::FAlignedFloatBuffer FMediaIOAudioOutput::GetFloatBuffer(uint32 NumSamples
 	return FloatBuffer;
 }
 
+PRAGMA_DISABLE_DEPRECATION_WARNINGS
 FMediaIOAudioCapture::FMediaIOAudioCapture(const FAudioDeviceHandle& InAudioDeviceHandle)
 {
-	RegisterAudioDevice(InAudioDeviceHandle);
 }
+PRAGMA_ENABLE_DEPRECATION_WARNINGS
 
 FMediaIOAudioCapture::~FMediaIOAudioCapture()
 {
 	UnregisterAudioDevice();
 }
 
+void FMediaIOAudioCapture::Initialize(const FAudioDeviceHandle& InAudioDeviceHandle)
+{
+	RegisterAudioDevice(InAudioDeviceHandle);
+}
+
 void FMediaIOAudioCapture::OnNewSubmixBuffer(const USoundSubmix* InOwningSubmix, float* InAudioData, int32 InNumSamples, int32 InNumChannels, const int32 InSampleRate, double InAudioClock)
 {
 	TRACE_CPUPROFILER_EVENT_SCOPE(FMediaIOAudioCapture::OnNewSubmixBuffer);
@@ -206,11 +212,7 @@ void FMediaIOAudioCapture::UnregisterAudioDevice()
 }
 
 FMainMediaIOAudioCapture::FMainMediaIOAudioCapture()
-	: FMediaIOAudioCapture(FAudioDeviceHandle())
-{
-	// Register the current device (PIE or main).
-	RegisterCurrentAudioDevice();
-	
+{	
 #if WITH_EDITOR
 	FEditorDelegates::PostPIEStarted.AddRaw(this, &FMainMediaIOAudioCapture::OnPIEStarted);
 	FEditorDelegates::PrePIEEnded.AddRaw(this, &FMainMediaIOAudioCapture::OnPIEEnded);
@@ -225,6 +227,12 @@ FMainMediaIOAudioCapture::~FMainMediaIOAudioCapture()
 #endif
 }
 
+void FMainMediaIOAudioCapture::Initialize()
+{
+	// Register the current device (PIE or main).
+	RegisterCurrentAudioDevice();
+}
+
 #if WITH_EDITOR
 void FMainMediaIOAudioCapture::OnPIEStarted(const bool)
 {
diff --git a/Engine/Plugins/Media/MediaIOFramework/Source/MediaIOCore/Private/MediaIOCoreSubsystem.cpp b/Engine/Plugins/Media/MediaIOFramework/Source/MediaIOCore/Private/MediaIOCoreSubsystem.cpp
--- a/Engine/Plugins/Media/MediaIOFramework/Source/MediaIOCore/Private/MediaIOCoreSubsystem.cpp
+++ b/Engine/Plugins/Media/MediaIOFramework/Source/MediaIOCore/Private/MediaIOCoreSubsystem.cpp
@@ -17,17 +17,18 @@ void UMediaIOCoreSubsystem::Deinitialize()
 
 TSharedPtr<FMediaIOAudioOutput> UMediaIOCoreSubsystem::CreateAudioOutput(const FCreateAudioOutputArgs& InArgs)
 {
-	FMediaIOAudioCapture* MediaIOAudioCapture = nullptr;
+	TSharedPtr<FMediaIOAudioCapture> MediaIOAudioCapture;
 
 	if (InArgs.AudioDeviceHandle.IsValid())
 	{
 		if (const TSharedPtr<FMediaIOAudioCapture, ESPMode::ThreadSafe>* FoundMediaIOAudioCapture = MediaIOAudioCaptures.Find(InArgs.AudioDeviceHandle.GetDeviceID()))
 		{
-			MediaIOAudioCapture = FoundMediaIOAudioCapture->Get();
+			MediaIOAudioCapture = *FoundMediaIOAudioCapture;
 		}
 		else
 		{
-			MediaIOAudioCapture = MediaIOAudioCaptures.Add(InArgs.AudioDeviceHandle.GetDeviceID(), MakeShared<FMediaIOAudioCapture, ESPMode::ThreadSafe>(InArgs.AudioDeviceHandle)).Get();
+			MediaIOAudioCapture = MediaIOAudioCaptures.Add(InArgs.AudioDeviceHandle.GetDeviceID(), MakeShared<FMediaIOAudioCapture>());
+			MediaIOAudioCapture->Initialize(InArgs.AudioDeviceHandle);
 			MainMediaIOAudioCapture->OnAudioCaptured_RenderThread().BindUObject(this, &UMediaIOCoreSubsystem::OnBufferReceivedByCapture, InArgs.AudioDeviceHandle.GetDeviceID());
 		}
 	}
@@ -36,10 +37,11 @@ TSharedPtr<FMediaIOAudioOutput> UMediaIOCoreSubsystem::CreateAudioOutput(const F
 		// Fallback using the main audio device.
 		if (!MainMediaIOAudioCapture)
 		{
-			MainMediaIOAudioCapture = MakeUnique<FMainMediaIOAudioCapture>();
+			MainMediaIOAudioCapture = MakeShared<FMainMediaIOAudioCapture>();
+			MainMediaIOAudioCapture->Initialize();
 			MainMediaIOAudioCapture->OnAudioCaptured_RenderThread().BindUObject(this, &UMediaIOCoreSubsystem::OnBufferReceivedByCapture, InArgs.AudioDeviceHandle.GetDeviceID());
 		}
-		MediaIOAudioCapture = MainMediaIOAudioCapture.Get();
+		MediaIOAudioCapture = MainMediaIOAudioCapture;
 	}	
 
 	return MediaIOAudioCapture->CreateAudioOutput(InArgs.NumOutputChannels, InArgs.TargetFrameRate, InArgs.MaxSampleLatency, InArgs.OutputSampleRate);
diff --git a/Engine/Plugins/Media/MediaIOFramework/Source/MediaIOCore/Public/MediaIOCoreAudioOutput.h b/Engine/Plugins/Media/MediaIOFramework/Source/MediaIOCore/Public/MediaIOCoreAudioOutput.h
--- a/Engine/Plugins/Media/MediaIOFramework/Source/MediaIOCore/Public/MediaIOCoreAudioOutput.h
+++ b/Engine/Plugins/Media/MediaIOFramework/Source/MediaIOCore/Public/MediaIOCoreAudioOutput.h
@@ -158,13 +158,20 @@ class FMediaIOAudioCapture : public ISubmixBufferListener
 public:
 	DECLARE_DELEGATE_TwoParams(FOnAudioCaptured, float* /* data */, int32 /* NumSamples */);
 
+	FMediaIOAudioCapture() = default;
+
+	UE_DEPRECATED(5.4, "This constructor is deprecated. Use Initialize(InAudioDeviceHandle) instead after constructing it.")
 	FMediaIOAudioCapture(const FAudioDeviceHandle& InAudioDeviceHandle);
+
 	virtual ~FMediaIOAudioCapture();
 
 	//~ ISubmixBufferListener interface
 	virtual void OnNewSubmixBuffer(const USoundSubmix* InOwningSubmix, float* InAudioData, int32 InNumSamples, int32 InNumChannels, const int32 InSampleRate, double InAudioClock) override;
 	virtual const FString& GetListenerName() const override;
 
+	/** Initializes audio capture for the given audio device */
+	void Initialize(const FAudioDeviceHandle& InAudioDeviceHandle);
+
 	/** Create an audio output that will receive audio samples. */
 	TSharedPtr<FMediaIOAudioOutput> CreateAudioOutput(int32 InNumOutputChannels, FFrameRate InTargetFrameRate, uint32 InMaxSampleLatency, uint32 InOutputSampleRate);
 
@@ -217,6 +224,9 @@ public:
 	FMainMediaIOAudioCapture();
 	virtual ~FMainMediaIOAudioCapture() override;
 
+	/** Initializes audio capture for the main audio device */
+	void Initialize();
+
 private:
 #if WITH_EDITOR
 	void OnPIEStarted(const bool);
diff --git a/Engine/Plugins/Media/MediaIOFramework/Source/MediaIOCore/Public/MediaIOCoreSubsystem.h b/Engine/Plugins/Media/MediaIOFramework/Source/MediaIOCore/Public/MediaIOCoreSubsystem.h
--- a/Engine/Plugins/Media/MediaIOFramework/Source/MediaIOCore/Public/MediaIOCoreSubsystem.h
+++ b/Engine/Plugins/Media/MediaIOFramework/Source/MediaIOCore/Public/MediaIOCoreSubsystem.h
@@ -56,7 +56,7 @@ private:
 	void OnBufferReceivedByCapture(float* Data, int32 NumSamples, Audio::FDeviceId AudioDeviceID) const;
 
 private:
-	TUniquePtr<FMediaIOAudioCapture> MainMediaIOAudioCapture;
+	TSharedPtr<FMainMediaIOAudioCapture> MainMediaIOAudioCapture;
 	
 	TMap<Audio::FDeviceId, TSharedPtr<FMediaIOAudioCapture, ESPMode::ThreadSafe>> MediaIOAudioCaptures;
 
