Fixed ASan "unreachable code" error in MediaIOFramework plugin.

diff --git a/Engine/Plugins/Media/MediaIOFramework/Source/MediaIOCore/Private/Player/MediaIOCorePlayerBase.cpp b/Engine/Plugins/Media/MediaIOFramework/Source/MediaIOCore/Private/Player/MediaIOCorePlayerBase.cpp
--- a/Engine/Plugins/Media/MediaIOFramework/Source/MediaIOCore/Private/Player/MediaIOCorePlayerBase.cpp
+++ b/Engine/Plugins/Media/MediaIOFramework/Source/MediaIOCore/Private/Player/MediaIOCorePlayerBase.cpp
@@ -467,8 +467,6 @@ FText FMediaIOCorePlayerBase::GetTrackDisplayName(EMediaTrackType TrackType, int
 	default:
 		return FText::GetEmpty();
 	}
-
-	return FText::GetEmpty();
 }
 
 
