Re-fixing "unreachable code" since it was fixed differently in 28.20.

diff --git a/Engine/Plugins/Media/MediaIOFramework/Source/MediaIOCore/Private/Player/MediaIOCorePlayerBase.cpp b/Engine/Plugins/Media/MediaIOFramework/Source/MediaIOCore/Private/Player/MediaIOCorePlayerBase.cpp
--- a/Engine/Plugins/Media/MediaIOFramework/Source/MediaIOCore/Private/Player/MediaIOCorePlayerBase.cpp
+++ b/Engine/Plugins/Media/MediaIOFramework/Source/MediaIOCore/Private/Player/MediaIOCorePlayerBase.cpp
@@ -467,6 +467,8 @@ FText FMediaIOCorePlayerBase::GetTrackDisplayName(EMediaTrackType TrackType, int
 	default:
 		break;
 	}
+
+	return FText::GetEmpty();
 }
 
 
