Fix memory lifetime issues in MediaIOFramework / PixelStreaming
 - Texture debug names cannot be generated dynamically via printf, as the underlying class only holds a TCHAR*.
 - Lambda passed to FRenderPass::FExecutePass::CreateLambda should take the CapturingFrame TSharedPtr arg by reference, as the caller owns the pointer, to avoid an unnecessary addref/release.

diff --git a/Engine/Plugins/Media/MediaIOFramework/Source/MediaIOCore/Private/MediaCaptureRenderPass.cpp b/Engine/Plugins/Media/MediaIOFramework/Source/MediaIOCore/Private/MediaCaptureRenderPass.cpp
--- a/Engine/Plugins/Media/MediaIOFramework/Source/MediaIOCore/Private/MediaCaptureRenderPass.cpp
+++ b/Engine/Plugins/Media/MediaIOFramework/Source/MediaIOCore/Private/MediaCaptureRenderPass.cpp
@@ -157,7 +157,7 @@ namespace UE::MediaCapture
 		ColorConversionOutputTextureDesc.ClearValue = FClearValueBinding(FLinearColor::White);
 		ColorConversionOutputTextureDesc.Reset();
 
-		TRefCountPtr<IPooledRenderTarget> RenderTarget = Args.MediaCapture->InitializePassOutputTexture(ColorConversionOutputTextureDesc, FString::Printf(TEXT("MediaCapture ColorConversion RenderTarget %d"), FrameId));
+		TRefCountPtr<IPooledRenderTarget> RenderTarget = Args.MediaCapture->InitializePassOutputTexture(ColorConversionOutputTextureDesc, TEXT("MediaCapture ColorConversion RenderTarget"));
 		return RenderTargetResource(MoveTemp(RenderTarget));
 	}
 
@@ -175,7 +175,7 @@ namespace UE::MediaCapture
 		ResampleOutputTextureDesc.Flags |= TexCreate_RenderTargetable | TexCreate_UAV | TexCreate_NoFastClear;
 		ResampleOutputTextureDesc.ClearValue = FClearValueBinding(FLinearColor(0, 0, 0, 0));
 
-		TRefCountPtr<IPooledRenderTarget> RenderTarget = Args.MediaCapture->InitializePassOutputTexture(ResampleOutputTextureDesc, FString::Printf(TEXT("MediaCapture Resample RenderTarget %d"), FrameId));
+		TRefCountPtr<IPooledRenderTarget> RenderTarget = Args.MediaCapture->InitializePassOutputTexture(ResampleOutputTextureDesc, TEXT("MediaCapture Resample RenderTarget"));
 		return RenderTargetResource(MoveTemp(RenderTarget));
 	}
 
@@ -245,7 +245,7 @@ namespace UE::MediaCapture
 		ResamplePass.Name = "Resample";
 		ResamplePass.OutputType = ERDGViewableResourceType::Texture;
 		ResamplePass.InitializePassOutputDelegate = FRenderPass::FInitializePassOutput::CreateStatic(&InitializeResamplePassOutputTexture);
-		ResamplePass.ExecutePassDelegate = FRenderPass::FExecutePass::CreateLambda([](const UE::MediaCaptureData::FCaptureFrameArgs& Args, const TSharedPtr<UE::MediaCaptureData::FCaptureFrame> CapturingFrame, FRDGViewableResource* InputResource, FRDGViewableResource* OutputTexture)
+		ResamplePass.ExecutePassDelegate = FRenderPass::FExecutePass::CreateLambda([](const UE::MediaCaptureData::FCaptureFrameArgs& Args, const TSharedPtr<UE::MediaCaptureData::FCaptureFrame>& CapturingFrame, FRDGViewableResource* InputResource, FRDGViewableResource* OutputTexture)
 		{
 			check(InputResource->Type == ERDGViewableResourceType::Texture);
 
@@ -269,7 +269,7 @@ namespace UE::MediaCapture
 		ColorConversionPass.Name = "ColorConversion";
 		ColorConversionPass.OutputType = ERDGViewableResourceType::Texture;
 		ColorConversionPass.InitializePassOutputDelegate = FRenderPass::FInitializePassOutput::CreateStatic(&InitializeColorConversionPassOutputTexture);
-		ColorConversionPass.ExecutePassDelegate = FRenderPass::FExecutePass::CreateLambda([CachedOCIOResources = UE::MediaCapture::ColorConversion::GetColorConversionResources(MediaCapture)](const UE::MediaCaptureData::FCaptureFrameArgs& Args, const TSharedPtr<UE::MediaCaptureData::FCaptureFrame> CapturingFrame, FRDGViewableResource* InputResource, FRDGViewableResource* OutputTexture)
+		ColorConversionPass.ExecutePassDelegate = FRenderPass::FExecutePass::CreateLambda([CachedOCIOResources = UE::MediaCapture::ColorConversion::GetColorConversionResources(MediaCapture)](const UE::MediaCaptureData::FCaptureFrameArgs& Args, const TSharedPtr<UE::MediaCaptureData::FCaptureFrame>& CapturingFrame, FRDGViewableResource* InputResource, FRDGViewableResource* OutputTexture)
 		{
 			check(InputResource->Type == ERDGViewableResourceType::Texture);
 
@@ -312,7 +312,7 @@ namespace UE::MediaCapture
 		if (MediaCapture->GetConversionOperation() == EMediaCaptureConversionOperation::CUSTOM)
 		{
 			ConversionPass.Name = "CustomConversion";
-			ConversionPass.ExecutePassDelegate = FRenderPass::FExecutePass::CreateLambda([](const UE::MediaCaptureData::FCaptureFrameArgs& Args, const TSharedPtr<UE::MediaCaptureData::FCaptureFrame> CapturingFrame, FRDGViewableResource* InputResource, FRDGViewableResource* OutputResource)
+			ConversionPass.ExecutePassDelegate = FRenderPass::FExecutePass::CreateLambda([](const UE::MediaCaptureData::FCaptureFrameArgs& Args, const TSharedPtr<UE::MediaCaptureData::FCaptureFrame>& CapturingFrame, FRDGViewableResource* InputResource, FRDGViewableResource* OutputResource)
 			{
 				// We only support Texture inputs
 				check(InputResource->Type == ERDGViewableResourceType::Texture);
@@ -338,7 +338,7 @@ namespace UE::MediaCapture
 		else
 		{
 			ConversionPass.Name = "Conversion";
-			ConversionPass.ExecutePassDelegate = FRenderPass::FExecutePass::CreateLambda([](const UE::MediaCaptureData::FCaptureFrameArgs& Args, const TSharedPtr<UE::MediaCaptureData::FCaptureFrame> CapturingFrame, FRDGViewableResource* InputResource, FRDGViewableResource* OutputResource)
+			ConversionPass.ExecutePassDelegate = FRenderPass::FExecutePass::CreateLambda([](const UE::MediaCaptureData::FCaptureFrameArgs& Args, const TSharedPtr<UE::MediaCaptureData::FCaptureFrame>& CapturingFrame, FRDGViewableResource* InputResource, FRDGViewableResource* OutputResource)
 			{
 				// At the moment we only support Texture inputs
 				check(InputResource->Type == ERDGViewableResourceType::Texture);
diff --git a/Engine/Plugins/Media/MediaIOFramework/Source/MediaIOCore/Public/MediaCapture.h b/Engine/Plugins/Media/MediaIOFramework/Source/MediaIOCore/Public/MediaCapture.h
--- a/Engine/Plugins/Media/MediaIOFramework/Source/MediaIOCore/Public/MediaCapture.h
+++ b/Engine/Plugins/Media/MediaIOFramework/Source/MediaIOCore/Public/MediaCapture.h
@@ -414,14 +414,14 @@ public:
 
 	static const TSet<EPixelFormat>& GetSupportedRgbaSwizzleFormats();
 
-	virtual TRefCountPtr<IPooledRenderTarget> InitializePassOutputTexture(FRDGTextureDesc TextureDesc, const FString& TextureName) const
+	virtual TRefCountPtr<IPooledRenderTarget> InitializePassOutputTexture(FRDGTextureDesc TextureDesc, const TCHAR* TextureName) const
 	{
-		return AllocatePooledTexture(TextureDesc, *TextureName);
+		return AllocatePooledTexture(TextureDesc, TextureName);
 	}
 
-	virtual TRefCountPtr<FRDGPooledBuffer> InitializePassOutputBuffer(FRDGBufferDesc BufferDesc, const FString& BufferName) const
+	virtual TRefCountPtr<FRDGPooledBuffer> InitializePassOutputBuffer(FRDGBufferDesc BufferDesc, const TCHAR* BufferName) const
 	{
-		return AllocatePooledBuffer(BufferDesc, *BufferName);
+		return AllocatePooledBuffer(BufferDesc, BufferName);
 	}
 
 	/** Get the name of the media output that created this capture. */
diff --git a/Engine/Plugins/Media/PixelStreaming/Source/PixelStreaming/Internal/PixelStreamingMediaIOCapture.h b/Engine/Plugins/Media/PixelStreaming/Source/PixelStreaming/Internal/PixelStreamingMediaIOCapture.h
--- a/Engine/Plugins/Media/PixelStreaming/Source/PixelStreaming/Internal/PixelStreamingMediaIOCapture.h
+++ b/Engine/Plugins/Media/PixelStreaming/Source/PixelStreaming/Internal/PixelStreamingMediaIOCapture.h
@@ -61,7 +61,7 @@ public:
     virtual ETextureCreateFlags GetOutputTextureFlags() const override;
 	// We override the initialization of the pass output texture as we want to the `CreateRenderTarget` allocator as it doesn't force
 	// the SRV flag
-	virtual TRefCountPtr<IPooledRenderTarget> InitializePassOutputTexture(FRDGTextureDesc TextureDesc, const FString& TextureName) const override;
+	virtual TRefCountPtr<IPooledRenderTarget> InitializePassOutputTexture(FRDGTextureDesc TextureDesc, const TCHAR* TextureName) const override;
 	//~ End UMediaCapture interface
 
 	TSharedPtr<FSceneViewport> GetViewport() const { return SceneViewport.Pin(); }
diff --git a/Engine/Plugins/Media/PixelStreaming/Source/PixelStreaming/Private/PixelStreamingMediaIOCapture.cpp b/Engine/Plugins/Media/PixelStreaming/Source/PixelStreaming/Private/PixelStreamingMediaIOCapture.cpp
--- a/Engine/Plugins/Media/PixelStreaming/Source/PixelStreaming/Private/PixelStreamingMediaIOCapture.cpp
+++ b/Engine/Plugins/Media/PixelStreaming/Source/PixelStreaming/Private/PixelStreamingMediaIOCapture.cpp
@@ -177,15 +177,15 @@ ETextureCreateFlags UPixelStreamingMediaIOCapture::GetOutputTextureFlags() const
 #endif
 }
 
-TRefCountPtr<IPooledRenderTarget> UPixelStreamingMediaIOCapture::InitializePassOutputTexture(FRDGTextureDesc TextureDesc, const FString& TextureName) const
+TRefCountPtr<IPooledRenderTarget> UPixelStreamingMediaIOCapture::InitializePassOutputTexture(FRDGTextureDesc TextureDesc, const TCHAR* TextureName) const
 {
-	FRHITextureCreateDesc Desc = FRHITextureCreateDesc::Create2D(*TextureName, TextureDesc.Extent.X, TextureDesc.Extent.Y, TextureDesc.Format)
+	FRHITextureCreateDesc Desc = FRHITextureCreateDesc::Create2D(TextureName, TextureDesc.Extent.X, TextureDesc.Extent.Y, TextureDesc.Format)
 		.SetClearValue(TextureDesc.ClearValue)
 		.SetFlags(TextureDesc.Flags);
 
 	FTextureRHIRef Texture = RHICreateTexture(Desc);
 
-	return CreateRenderTarget(Texture, *TextureName);
+	return CreateRenderTarget(Texture, TextureName);
 }
 
 bool UPixelStreamingMediaIOCapture::PostInitializeCaptureViewport(TSharedPtr<FSceneViewport>& InSceneViewport)
