Fix: Memory leak on Mac introduced in CL31283735

diff --git a/Engine/Plugins/Media/MediaIOFramework/Source/MediaIOCore/Private/MediaCaptureRenderPass.cpp b/Engine/Plugins/Media/MediaIOFramework/Source/MediaIOCore/Private/MediaCaptureRenderPass.cpp
--- a/Engine/Plugins/Media/MediaIOFramework/Source/MediaIOCore/Private/MediaCaptureRenderPass.cpp
+++ b/Engine/Plugins/Media/MediaIOFramework/Source/MediaIOCore/Private/MediaCaptureRenderPass.cpp
@@ -157,7 +157,7 @@ namespace UE::MediaCapture
 		ColorConversionOutputTextureDesc.ClearValue = FClearValueBinding(FLinearColor::White);
 		ColorConversionOutputTextureDesc.Reset();
 
-		TRefCountPtr<IPooledRenderTarget> RenderTarget = Args.MediaCapture->InitializePassOutputTexture(ColorConversionOutputTextureDesc, TEXT("MediaCapture ColorConversion RenderTarget"));
+		TRefCountPtr<IPooledRenderTarget> RenderTarget = AllocatePooledTexture(ColorConversionOutputTextureDesc, TEXT("MediaCapture ColorConversion RenderTarget"));
 		return RenderTargetResource(MoveTemp(RenderTarget));
 	}
 
@@ -175,7 +175,7 @@ namespace UE::MediaCapture
 		ResampleOutputTextureDesc.Flags |= TexCreate_RenderTargetable | TexCreate_UAV | TexCreate_NoFastClear;
 		ResampleOutputTextureDesc.ClearValue = FClearValueBinding(FLinearColor(0, 0, 0, 0));
 
-		TRefCountPtr<IPooledRenderTarget> RenderTarget = Args.MediaCapture->InitializePassOutputTexture(ResampleOutputTextureDesc, TEXT("MediaCapture Resample RenderTarget"));
+		TRefCountPtr<IPooledRenderTarget> RenderTarget = AllocatePooledTexture(ResampleOutputTextureDesc, TEXT("MediaCapture Resample RenderTarget"));
 		return RenderTargetResource(MoveTemp(RenderTarget));
 	}
 
@@ -296,7 +296,7 @@ namespace UE::MediaCapture
 		{
 			ConversionPass.InitializePassOutputDelegate = FRenderPass::FInitializePassOutput::CreateLambda([](const UE::MediaCapture::FInitializePassOutputArgs& Args, uint32 FrameId)
 			{
-				TRefCountPtr<IPooledRenderTarget> RenderTarget = Args.MediaCapture->InitializePassOutputTexture(Args.MediaCapture->DesiredOutputTextureDescription, TEXT("MediaCapture Output RenderTarget"));
+				TRefCountPtr<IPooledRenderTarget> RenderTarget = AllocatePooledTexture(Args.MediaCapture->DesiredOutputTextureDescription, TEXT("MediaCapture Output RenderTarget"));
 				return RenderTargetResource(MoveTemp(RenderTarget));
 			});
 		}
@@ -304,7 +304,7 @@ namespace UE::MediaCapture
 		{
 			ConversionPass.InitializePassOutputDelegate = FRenderPass::FInitializePassOutput::CreateLambda([](const UE::MediaCapture::FInitializePassOutputArgs& Args, uint32 FrameId)
 			{
-				TRefCountPtr<FRDGPooledBuffer> Buffer = Args.MediaCapture->InitializePassOutputBuffer(Args.MediaCapture->DesiredOutputBufferDescription, TEXT("MediaCapture Output Buffer"));
+				TRefCountPtr<FRDGPooledBuffer> Buffer = AllocatePooledBuffer(Args.MediaCapture->DesiredOutputBufferDescription, TEXT("MediaCapture Output Buffer"));
 				return BufferResource(MoveTemp(Buffer));
 			});
 		}
diff --git a/Engine/Plugins/Media/MediaIOFramework/Source/MediaIOCore/Public/MediaCapture.h b/Engine/Plugins/Media/MediaIOFramework/Source/MediaIOCore/Public/MediaCapture.h
--- a/Engine/Plugins/Media/MediaIOFramework/Source/MediaIOCore/Public/MediaCapture.h
+++ b/Engine/Plugins/Media/MediaIOFramework/Source/MediaIOCore/Public/MediaCapture.h
@@ -414,16 +414,6 @@ public:
 
 	static const TSet<EPixelFormat>& GetSupportedRgbaSwizzleFormats();
 
-	virtual TRefCountPtr<IPooledRenderTarget> InitializePassOutputTexture(FRDGTextureDesc TextureDesc, const TCHAR* TextureName) const
-	{
-		return AllocatePooledTexture(TextureDesc, TextureName);
-	}
-
-	virtual TRefCountPtr<FRDGPooledBuffer> InitializePassOutputBuffer(FRDGBufferDesc BufferDesc, const TCHAR* BufferName) const
-	{
-		return AllocatePooledBuffer(BufferDesc, BufferName);
-	}
-
 	/** Get the name of the media output that created this capture. */
 	FString GetMediaOutputName() const
 	{
diff --git a/Engine/Plugins/Media/PixelStreaming/Source/PixelStreaming/Internal/PixelStreamingMediaIOCapture.h b/Engine/Plugins/Media/PixelStreaming/Source/PixelStreaming/Internal/PixelStreamingMediaIOCapture.h
--- a/Engine/Plugins/Media/PixelStreaming/Source/PixelStreaming/Internal/PixelStreamingMediaIOCapture.h
+++ b/Engine/Plugins/Media/PixelStreaming/Source/PixelStreaming/Internal/PixelStreamingMediaIOCapture.h
@@ -59,9 +59,6 @@ public:
 	virtual bool SupportsAnyThreadCapture() const override;
 	// We override the texture flags because on Mac we want the texture to have the CPU_Readback flag
     virtual ETextureCreateFlags GetOutputTextureFlags() const override;
-	// We override the initialization of the pass output texture as we want to the `CreateRenderTarget` allocator as it doesn't force
-	// the SRV flag
-	virtual TRefCountPtr<IPooledRenderTarget> InitializePassOutputTexture(FRDGTextureDesc TextureDesc, const TCHAR* TextureName) const override;
 	//~ End UMediaCapture interface
 
 	TSharedPtr<FSceneViewport> GetViewport() const { return SceneViewport.Pin(); }
diff --git a/Engine/Plugins/Media/PixelStreaming/Source/PixelStreaming/Private/PixelStreamingMediaIOCapture.cpp b/Engine/Plugins/Media/PixelStreaming/Source/PixelStreaming/Private/PixelStreamingMediaIOCapture.cpp
--- a/Engine/Plugins/Media/PixelStreaming/Source/PixelStreaming/Private/PixelStreamingMediaIOCapture.cpp
+++ b/Engine/Plugins/Media/PixelStreaming/Source/PixelStreaming/Private/PixelStreamingMediaIOCapture.cpp
@@ -177,17 +177,6 @@ ETextureCreateFlags UPixelStreamingMediaIOCapture::GetOutputTextureFlags() const
 #endif
 }
 
-TRefCountPtr<IPooledRenderTarget> UPixelStreamingMediaIOCapture::InitializePassOutputTexture(FRDGTextureDesc TextureDesc, const TCHAR* TextureName) const
-{
-	FRHITextureCreateDesc Desc = FRHITextureCreateDesc::Create2D(TextureName, TextureDesc.Extent.X, TextureDesc.Extent.Y, TextureDesc.Format)
-		.SetClearValue(TextureDesc.ClearValue)
-		.SetFlags(TextureDesc.Flags);
-
-	FTextureRHIRef Texture = RHICreateTexture(Desc);
-
-	return CreateRenderTarget(Texture, TextureName);
-}
-
 bool UPixelStreamingMediaIOCapture::PostInitializeCaptureViewport(TSharedPtr<FSceneViewport>& InSceneViewport)
 {
 	SceneViewport = TWeakPtr<FSceneViewport>(InSceneViewport);
