Add Misc/Paths.h for non-unity issue

#rnx

diff --git a/Engine/Plugins/Media/MediaIOFramework/Source/MediaIOCore/Public/MediaIOAudioDebug.h b/Engine/Plugins/Media/MediaIOFramework/Source/MediaIOCore/Public/MediaIOAudioDebug.h
--- a/Engine/Plugins/Media/MediaIOFramework/Source/MediaIOCore/Public/MediaIOAudioDebug.h
+++ b/Engine/Plugins/Media/MediaIOFramework/Source/MediaIOCore/Public/MediaIOAudioDebug.h
@@ -7,6 +7,7 @@
 #include "HAL/FileManager.h"
 #include "HAL/IConsoleManager.h"
 #include "Logging/LogVerbosity.h"
+#include "Misc/Paths.h"
 
 /**
  * Single header drop in file to help diagnose audio issues
