[MediaIOCore] QOL: Update CaptureImmediate_RenderThread to return a boolean that indicates if the capture was successfully processed.

diff --git a/Engine/Plugins/Media/MediaIOFramework/Source/MediaIOCore/Private/MediaCapture.cpp b/Engine/Plugins/Media/MediaIOFramework/Source/MediaIOCore/Private/MediaCapture.cpp
--- a/Engine/Plugins/Media/MediaIOFramework/Source/MediaIOCore/Private/MediaCapture.cpp
+++ b/Engine/Plugins/Media/MediaIOFramework/Source/MediaIOCore/Private/MediaCapture.cpp
@@ -603,6 +603,15 @@ bool UMediaCapture::SetCaptureAudioDevice(const FAudioDeviceHandle& InAudioDevic
 }
 
 void UMediaCapture::CaptureImmediate_RenderThread(FRDGBuilder& GraphBuilder, FRHITexture* InSourceTexture, FIntRect SourceViewRect)
+{
+	bool bCaptureSucceeded = TryCaptureImmediate_RenderThread(GraphBuilder, InSourceTexture, SourceViewRect);
+	if(!bCaptureSucceeded)
+	{
+		UE_LOG(LogMediaIOCore, VeryVerbose, TEXT("[%s] - Failed to capture resource."), *MediaOutputName);
+	}
+}
+
+bool UMediaCapture::TryCaptureImmediate_RenderThread(FRDGBuilder& GraphBuilder, FRHITexture* InSourceTexture, FIntRect SourceViewRect)
 {
 	UE::MediaCaptureData::FCaptureFrameArgs CaptureArgs{GraphBuilder};
 	CaptureArgs.MediaCapture = this;
@@ -610,11 +619,19 @@ void UMediaCapture::CaptureImmediate_RenderThread(FRDGBuilder& GraphBuilder, FRH
     CaptureArgs.ResourceToCapture = InSourceTexture;
 	CaptureArgs.SourceViewRect = SourceViewRect;
 	
-	CaptureImmediate_RenderThread(CaptureArgs);
+	return CaptureImmediate_RenderThread(CaptureArgs);
 }
 
-
 void UMediaCapture::CaptureImmediate_RenderThread(FRDGBuilder& GraphBuilder,  FRDGTextureRef InSourceTextureRef, FIntRect SourceViewRect)
+{
+	bool bCaptureSucceeded = TryCaptureImmediate_RenderThread(GraphBuilder, InSourceTextureRef, SourceViewRect);
+	if(!bCaptureSucceeded)
+	{
+		UE_LOG(LogMediaIOCore, VeryVerbose, TEXT("[%s] - Failed to capture resource."), *MediaOutputName);
+	}
+}
+
+bool UMediaCapture::TryCaptureImmediate_RenderThread(FRDGBuilder& GraphBuilder,  FRDGTextureRef InSourceTextureRef, FIntRect SourceViewRect)
 {
 	// Make sure captureimmediate runs are seen by game thread when stopping is requested. Having a pending render command will cause a flush
 	// hence letting the capture finish. If stopping is in progress, stopped state will be seen by the captureimmediate and early exit.
@@ -631,10 +648,10 @@ void UMediaCapture::CaptureImmediate_RenderThread(FRDGBuilder& GraphBuilder,  FR
 	CaptureArgs.RDGResourceToCapture = InSourceTextureRef;
 	CaptureArgs.SourceViewRect = SourceViewRect;
 	
-	CaptureImmediate_RenderThread(CaptureArgs);
+	return CaptureImmediate_RenderThread(CaptureArgs);
 }
 
-void UMediaCapture::CaptureImmediate_RenderThread(const UE::MediaCaptureData::FCaptureFrameArgs& Args)
+bool UMediaCapture::CaptureImmediate_RenderThread(const UE::MediaCaptureData::FCaptureFrameArgs& Args)
 {
 	using namespace UE::MediaCaptureData;
 	TRACE_CPUPROFILER_EVENT_SCOPE(UMediaCapture::CaptureImmediate_RenderThread);
@@ -644,24 +661,24 @@ void UMediaCapture::CaptureImmediate_RenderThread(const UE::MediaCaptureData::FC
 	
 	if (bIsAutoRestartRequired)
 	{
-		return;
+		return false;
 	}
 
 	// This could happen if the capture immediate is called before our resources were initialized. We can't really flush as we are in a command.
 	if (!bOutputResourcesInitialized)
 	{
 		UE_LOG(LogMediaIOCore, VeryVerbose, TEXT("Could not capture frame. Output resources haven't been initialized"));
-		return;
+		return false;
 	}
 
 	if (!MediaOutput)
 	{
-		return;
+		return false;
 	}
 
 	if (GetState() == EMediaCaptureState::Error)
 	{
-		return;
+		return false;
 	}
 
 	if (CaptureSource && CaptureSource->GetCaptureType() != UE::MediaCapture::Private::ECaptureType::Immediate)
@@ -672,19 +689,19 @@ void UMediaCapture::CaptureImmediate_RenderThread(const UE::MediaCaptureData::FC
 
 	if (GetState() != EMediaCaptureState::Capturing && GetState() != EMediaCaptureState::StopRequested)
 	{
-		return;
+		return false;
 	}
 
 	if (DesiredCaptureOptions.bSkipFrameWhenRunningExpensiveTasks && !FSlateThrottleManager::Get().IsAllowingExpensiveTasks())
 	{
-		return;
+		return false;
 	}
 
 	if (CaptureSource.IsValid() == false)
 	{
 		UE_LOG(LogMediaIOCore, Error, TEXT("[%s] - Trying to capture a RHI resource with an invalid source."), *MediaOutputName);
 		SetState(EMediaCaptureState::Error);
-		return;
+		return false;
 	}
 
 	if (CaptureSource->GetSourceType() == EMediaCaptureSourceType::RHI_RESOURCE)
@@ -719,14 +736,14 @@ void UMediaCapture::CaptureImmediate_RenderThread(const UE::MediaCaptureData::FC
 	if (bFoundMatchingData == false)
 	{
 		UE_LOG(LogMediaIOCore, Warning, TEXT("Can't capture frame. Could not find the matching game frame %d."), GFrameCounterRenderThread);
-		return;
+		return false;
 	}
 
 	if (GetState() == EMediaCaptureState::StopRequested && PendingFrameCount.load() <= 0)
 	{
 		// All the requested frames have been captured.
 		StopCapture(false);
-		return;
+		return false;
 	}
 		
 	// Get next available frame from the store. Can be invalid.
@@ -750,7 +767,7 @@ void UMediaCapture::CaptureImmediate_RenderThread(const UE::MediaCaptureData::FC
 					*MediaOutputName, GFrameCounterRenderThread);
 
 				SetState(EMediaCaptureState::Error);
-				return;
+				return false;
 			}
 		}
 		else
@@ -770,6 +787,7 @@ void UMediaCapture::CaptureImmediate_RenderThread(const UE::MediaCaptureData::FC
 	
 	UE_LOG(LogMediaIOCore, Verbose, TEXT("[%s - %s] - Capturing frame %d"), *MediaOutputName, *FThreadManager::GetThreadName(FPlatformTLS::GetCurrentThreadId()), CapturingFrame ? CapturingFrame->FrameId : -1);
 	
+	bool bCaptureSucceeded = false;
 	if (CapturingFrame)
 	{
 		// Prepare frame to capture
@@ -779,7 +797,7 @@ void UMediaCapture::CaptureImmediate_RenderThread(const UE::MediaCaptureData::FC
 		CapturingFrame->CaptureBaseData.SourceFrameNumber = ++CaptureRequestCount;
 		CapturingFrame->UserData = MoveTemp(NextCaptureData.UserData);
 
-		ProcessCapture_RenderThread(CapturingFrame, Args);
+		bCaptureSucceeded = ProcessCapture_RenderThread(CapturingFrame, Args);
 	}
 
 	// If CVarMediaIOEnableExperimentalScheduling is enabled, it means that a sync point pass was added, and the ready frame processing will be done later
@@ -788,9 +806,11 @@ void UMediaCapture::CaptureImmediate_RenderThread(const UE::MediaCaptureData::FC
 		if (TSharedPtr<FCaptureFrame> NextPending = FrameManager->PeekNextPending<FCaptureFrame>())
 		{
 			UE_LOG(LogMediaIOCore, Verbose, TEXT("[%s - %s] - Processing pending frame %d"), *MediaOutputName, *FThreadManager::GetThreadName(FPlatformTLS::GetCurrentThreadId()), NextPending ? NextPending->FrameId : -1);
-			ProcessReadyFrame_RenderThread(Args.GraphBuilder.RHICmdList, this, NextPending);
+			bCaptureSucceeded = ProcessReadyFrame_RenderThread(Args.GraphBuilder.RHICmdList, this, NextPending);
 		}
 	}
+
+	return bCaptureSucceeded;
 }
 
 UTextureRenderTarget2D* UMediaCapture::GetTextureRenderTarget() const
diff --git a/Engine/Plugins/Media/MediaIOFramework/Source/MediaIOCore/Private/MediaCaptureSceneViewExtension.h b/Engine/Plugins/Media/MediaIOFramework/Source/MediaIOCore/Private/MediaCaptureSceneViewExtension.h
--- a/Engine/Plugins/Media/MediaIOFramework/Source/MediaIOCore/Private/MediaCaptureSceneViewExtension.h
+++ b/Engine/Plugins/Media/MediaIOFramework/Source/MediaIOCore/Private/MediaCaptureSceneViewExtension.h
@@ -102,7 +102,12 @@ public:
 			RDG_EVENT_SCOPE(GraphBuilder, "MediaCaptureSceneExtension");
 			if (WeakCapture.IsValid())
 			{
-				WeakCapture->CaptureImmediate_RenderThread(GraphBuilder, TextureRef, SceneColor.ViewRect);
+				bool bCaptureSucceeded = WeakCapture->TryCaptureImmediate_RenderThread(GraphBuilder, TextureRef, SceneColor.ViewRect);
+				if(!bCaptureSucceeded)
+				{
+					LastErrorMessage = TEXT("Failed to capture resource.");
+					UE_LOG(LogMediaIOCore, VeryVerbose, TEXT("%s"), *LastErrorMessage);
+				}
 			}
 		}
 
diff --git a/Engine/Plugins/Media/MediaIOFramework/Source/MediaIOCore/Public/MediaCapture.h b/Engine/Plugins/Media/MediaIOFramework/Source/MediaIOCore/Public/MediaCapture.h
--- a/Engine/Plugins/Media/MediaIOFramework/Source/MediaIOCore/Public/MediaCapture.h
+++ b/Engine/Plugins/Media/MediaIOFramework/Source/MediaIOCore/Public/MediaCapture.h
@@ -323,9 +323,17 @@ public:
 	UFUNCTION(BlueprintCallable, Category = "Media|Output")
 	bool UpdateTextureRenderTarget2D(UTextureRenderTarget2D* RenderTarget);
 	
-	/** Captures a resource immediately from the render thread. Used in RHI_RESOURCE capture mode */
+	/** 
+	 * Captures a resource immediately from the render thread. Used in RHI_RESOURCE capture mode 
+	 * @return Return true if the capture was successfully processed.
+	 */
+	bool TryCaptureImmediate_RenderThread(FRDGBuilder& GraphBuilder, FRHITexture* InSourceTexture, FIntRect SourceViewRect = FIntRect(0,0,0,0));
+	bool TryCaptureImmediate_RenderThread(FRDGBuilder& GraphBuilder, FRDGTextureRef InSourceTextureRef, FIntRect SourceViewRect = FIntRect(0,0,0,0));
+
+	UE_DEPRECATED(5.5, "This method has been deprecated. Please use TryCaptureImmediate_RenderThread as it returns a boolean to signify that the capture was successfully processed.")
 	void CaptureImmediate_RenderThread(FRDGBuilder& GraphBuilder, FRHITexture* InSourceTexture, FIntRect SourceViewRect = FIntRect(0,0,0,0));
-	
+
+	UE_DEPRECATED(5.5, "This method has been deprecated. Please use TryCaptureImmediate_RenderThread as it returns a boolean to signify that the capture was successfully processed.")
 	void CaptureImmediate_RenderThread(FRDGBuilder& GraphBuilder, FRDGTextureRef InSourceTextureRef, FIntRect SourceViewRect = FIntRect(0,0,0,0));
 
 	/**
@@ -612,7 +620,7 @@ private:
 	void WaitForSingleExperimentalSchedulingTaskToComplete();
 	void WaitForAllExperimentalSchedulingTasksToComplete();
 	
-	void CaptureImmediate_RenderThread(const UE::MediaCaptureData::FCaptureFrameArgs& Args);
+	bool CaptureImmediate_RenderThread(const UE::MediaCaptureData::FCaptureFrameArgs& Args);
 	// Capture pipeline stuff
 	void ProcessCapture_GameThread();
 	void PrepareAndDispatchCapture_GameThread(const TSharedPtr<UE::MediaCaptureData::FCaptureFrame>& CapturingFrame);
diff --git a/Engine/Plugins/Runtime/nDisplay/Source/DisplayClusterMedia/Private/Capture/DisplayClusterMediaCaptureBase.cpp b/Engine/Plugins/Runtime/nDisplay/Source/DisplayClusterMedia/Private/Capture/DisplayClusterMediaCaptureBase.cpp
--- a/Engine/Plugins/Runtime/nDisplay/Source/DisplayClusterMedia/Private/Capture/DisplayClusterMediaCaptureBase.cpp
+++ b/Engine/Plugins/Runtime/nDisplay/Source/DisplayClusterMedia/Private/Capture/DisplayClusterMediaCaptureBase.cpp
@@ -144,7 +144,11 @@ void FDisplayClusterMediaCaptureBase::ExportMediaData_RenderThread(FRDGBuilder&
 		*GetMediaId(), SrcTextureSize.X, SrcTextureSize.Y, SrcRegionSize.X, SrcRegionSize.Y, GFrameCounterRenderThread);
 
 	// Capture
-	MediaCapture->CaptureImmediate_RenderThread(GraphBuilder, TextureInfo.Texture, TextureInfo.Region);
+	bool bCaptureSucceeded = MediaCapture->TryCaptureImmediate_RenderThread(GraphBuilder, TextureInfo.Texture, TextureInfo.Region);
+	if(!bCaptureSucceeded)
+	{
+		UE_LOG(LogDisplayClusterMedia, VeryVerbose, TEXT("MediaCapture '%s': failed to capture resource"), *GetMediaId());
+	}
 }
 
 void FDisplayClusterMediaCaptureBase::OnPostClusterTick()
