MediaIOFramework:
- temporary plugged the issue with OCIO color conversion being broken.
- Made sample picking methods inheritable so that Rivermax can inherit from MediaIO.
- Media texture sample now passes texture to write into from Media Texture Resource.
- Color conversion settings are now moved from AJA and Blackmagic into the base class.

[RN] MediaIO: Made it easier for new capture card media plugins to inherit from Media IO Core.

diff --git a/Engine/Plugins/Media/AjaMedia/Source/AjaMedia/Private/Player/AjaMediaPlayer.cpp b/Engine/Plugins/Media/AjaMedia/Source/AjaMedia/Private/Player/AjaMediaPlayer.cpp
--- a/Engine/Plugins/Media/AjaMedia/Source/AjaMedia/Private/Player/AjaMediaPlayer.cpp
+++ b/Engine/Plugins/Media/AjaMedia/Source/AjaMedia/Private/Player/AjaMediaPlayer.cpp
@@ -103,11 +103,6 @@ bool FAjaMediaPlayer::Open(const FString& Url, const IMediaOptions* Options)
 	const bool bAutoDetectTimecode = Timecode == EMediaIOAutoDetectableTimecodeFormat::Auto;
 	const bool bAutoDetectVideoFormat = bAutoDetect;
 
-	bOverrideSourceEncoding = Options->GetMediaOption(UE::CaptureCardMediaSource::OverrideSourceEncoding, true);
-	OverrideSourceEncoding = (ETextureSourceEncoding) Options->GetMediaOption(UE::CaptureCardMediaSource::SourceEncoding, (int64) ETextureSourceEncoding::TSE_Linear);
-	bOverrideSourceColorSpace = Options->GetMediaOption(UE::CaptureCardMediaSource::OverrideSourceColorSpace, true);
-	OverrideSourceColorSpace =  (ETextureColorSpace) Options->GetMediaOption(UE::CaptureCardMediaSource::SourceColorSpace, (int64) ETextureColorSpace::TCS_None);
-
 	if (!bAutoDetectTimecode)
 	{
 		TimecodeFormat = UE::MediaIO::FromAutoDetectableTimecodeFormat(Timecode);
diff --git a/Engine/Plugins/Media/AjaMedia/Source/AjaMedia/Private/Player/AjaMediaPlayer.h b/Engine/Plugins/Media/AjaMedia/Source/AjaMedia/Private/Player/AjaMediaPlayer.h
--- a/Engine/Plugins/Media/AjaMedia/Source/AjaMedia/Private/Player/AjaMediaPlayer.h
+++ b/Engine/Plugins/Media/AjaMedia/Source/AjaMedia/Private/Player/AjaMediaPlayer.h
@@ -192,18 +192,6 @@ private:
 
 	EAjaMediaSourceColorFormat AjaColorFormat = EAjaMediaSourceColorFormat::YUV2_8bit;
 
-	/** Whether to override the source encoding or to use the metadata embedded in the ancillary data of the signal. */
-	bool bOverrideSourceEncoding = true;
-
-	/** Encoding of the source texture. */
-	ETextureSourceEncoding OverrideSourceEncoding = ETextureSourceEncoding::TSE_Linear;
-
-	/** Whether to override the source color space or to use the metadata embedded in the ancillary data of the signal. */
-	bool bOverrideSourceColorSpace = true;
-
-	/** Color space of the source texture. */
-	ETextureColorSpace OverrideSourceColorSpace = ETextureColorSpace::TCS_None;
-
 	/** Device provider used to autodetect input format. */
 	TPimplPtr<class FAjaDeviceProvider> DeviceProvider;
 
diff --git a/Engine/Plugins/Media/BlackmagicMedia/Source/BlackmagicMedia/Private/Player/BlackmagicMediaPlayer.cpp b/Engine/Plugins/Media/BlackmagicMedia/Source/BlackmagicMedia/Private/Player/BlackmagicMediaPlayer.cpp
--- a/Engine/Plugins/Media/BlackmagicMedia/Source/BlackmagicMedia/Private/Player/BlackmagicMediaPlayer.cpp
+++ b/Engine/Plugins/Media/BlackmagicMedia/Source/BlackmagicMedia/Private/Player/BlackmagicMediaPlayer.cpp
@@ -559,11 +559,6 @@ bool FBlackmagicMediaPlayer::Open(const FString& Url, const IMediaOptions* Optio
 	
 	const EMediaIOAutoDetectableTimecodeFormat TimecodeFormat = (EMediaIOAutoDetectableTimecodeFormat)(Options->GetMediaOption(BlackmagicMediaOption::TimecodeFormat, (int64)EMediaIOAutoDetectableTimecodeFormat::None));
 	const bool bAutoDetectTimecode = TimecodeFormat == EMediaIOAutoDetectableTimecodeFormat::Auto;
-
-	bOverrideSourceEncoding = Options->GetMediaOption(UE::CaptureCardMediaSource::OverrideSourceEncoding, true);
-	OverrideSourceEncoding = (ETextureSourceEncoding) Options->GetMediaOption(UE::CaptureCardMediaSource::SourceEncoding, (int64) ETextureSourceEncoding::TSE_Linear);
-	bOverrideSourceColorSpace = Options->GetMediaOption(UE::CaptureCardMediaSource::OverrideSourceColorSpace, true);
-	OverrideSourceColorSpace =  (ETextureColorSpace) Options->GetMediaOption(UE::CaptureCardMediaSource::SourceColorSpace, (int64) ETextureColorSpace::TCS_None);
 	
 	BlackmagicDesign::FChannelInfo ChannelInfo;
 	ChannelInfo.DeviceIndex = Options->GetMediaOption(BlackmagicMediaOption::DeviceIndex, (int64)0);
diff --git a/Engine/Plugins/Media/BlackmagicMedia/Source/BlackmagicMedia/Private/Player/BlackmagicMediaPlayer.h b/Engine/Plugins/Media/BlackmagicMedia/Source/BlackmagicMedia/Private/Player/BlackmagicMediaPlayer.h
--- a/Engine/Plugins/Media/BlackmagicMedia/Source/BlackmagicMedia/Private/Player/BlackmagicMediaPlayer.h
+++ b/Engine/Plugins/Media/BlackmagicMedia/Source/BlackmagicMedia/Private/Player/BlackmagicMediaPlayer.h
@@ -159,16 +159,4 @@ private:
 	EMediaIOSampleType SupportedSampleTypes;
 
 	EBlackmagicMediaSourceColorFormat BlackmagicColorFormat = EBlackmagicMediaSourceColorFormat::YUV8;
-
-	/** Whether to override the source encoding or to use the metadata embedded in the ancillary data of the signal. */
-	bool bOverrideSourceEncoding = true;
-
-	/** Encoding of the source texture. */
-	ETextureSourceEncoding OverrideSourceEncoding = ETextureSourceEncoding::TSE_Linear;
-
-	/** Whether to override the source color space or to use the metadata embedded in the ancillary data of the signal. */
-	bool bOverrideSourceColorSpace = true;
-
-	/** Color space of the source texture. */
-	ETextureColorSpace OverrideSourceColorSpace = ETextureColorSpace::TCS_None;
 };
diff --git a/Engine/Plugins/Media/MediaIOFramework/Source/MediaIOCore/Private/MediaIOCoreTextureSampleConverter.cpp b/Engine/Plugins/Media/MediaIOFramework/Source/MediaIOCore/Private/MediaIOCoreTextureSampleConverter.cpp
--- a/Engine/Plugins/Media/MediaIOFramework/Source/MediaIOCore/Private/MediaIOCoreTextureSampleConverter.cpp
+++ b/Engine/Plugins/Media/MediaIOFramework/Source/MediaIOCore/Private/MediaIOCoreTextureSampleConverter.cpp
@@ -31,7 +31,7 @@ bool FMediaIOCoreTextureSampleConverter::Convert(FRHICommandListImmediate& RHICm
 	}
 
 	// Let the player pick a proper sample and render it into this proxy's texture (JITR)
-	return PlayerPtr->JustInTimeSampleRender_RenderThread(RHICmdList, SamplePtr);
+	return PlayerPtr->JustInTimeSampleRender_RenderThread(RHICmdList, InDestinationTexture, SamplePtr);
 }
 
 uint32 FMediaIOCoreTextureSampleConverter::GetConverterInfoFlags() const
diff --git a/Engine/Plugins/Media/MediaIOFramework/Source/MediaIOCore/Private/Player/MediaIOCorePlayerBase.cpp b/Engine/Plugins/Media/MediaIOFramework/Source/MediaIOCore/Private/Player/MediaIOCorePlayerBase.cpp
--- a/Engine/Plugins/Media/MediaIOFramework/Source/MediaIOCore/Private/Player/MediaIOCorePlayerBase.cpp
+++ b/Engine/Plugins/Media/MediaIOFramework/Source/MediaIOCore/Private/Player/MediaIOCorePlayerBase.cpp
@@ -214,7 +214,7 @@ IMediaView& FMediaIOCorePlayerBase::GetView()
 bool FMediaIOCorePlayerBase::Open(const FString& Url, const IMediaOptions* Options)
 {
 	Close();
-
+	CurrentState = EMediaState::Error;
 	ITimeManagementModule::Get().GetTimedDataInputCollection().Add(this);
 
 	OpenUrl = Url;
@@ -243,6 +243,17 @@ bool FMediaIOCorePlayerBase::Open(const FString& Url, const IMediaOptions* Optio
 		Deinterlacer = MakeShared<UE::MediaIOCore::FDeinterlacer>(UE::MediaIOCore::FDeinterlacer::FOnAcquireSample_AnyThread::CreateSP(this, &FMediaIOCorePlayerBase::AcquireTextureSample_AnyThread), InterlaceFieldOrder);
 	}
 
+	bOverrideSourceEncoding = Options->GetMediaOption(UE::CaptureCardMediaSource::OverrideSourceEncoding, true);
+	OverrideSourceEncoding = (ETextureSourceEncoding)Options->GetMediaOption(UE::CaptureCardMediaSource::SourceEncoding, (int64)ETextureSourceEncoding::TSE_Linear);
+	bOverrideSourceColorSpace = Options->GetMediaOption(UE::CaptureCardMediaSource::OverrideSourceColorSpace, true);
+	OverrideSourceColorSpace = (ETextureColorSpace)Options->GetMediaOption(UE::CaptureCardMediaSource::SourceColorSpace, (int64)ETextureColorSpace::TCS_None);
+
+	if (bReadMediaOptions)
+	{
+		CurrentState = EMediaState::Preparing;
+		EventSink.ReceiveMediaEvent(EMediaEvent::MediaConnecting);
+	}
+
 	return bReadMediaOptions;
 }
 
@@ -278,7 +289,7 @@ void FMediaIOCorePlayerBase::TickTimeManagement()
 		}
 		else
 		{
-			UE_LOG(LogMediaIOCore, Verbose, TEXT("The video '%s' is configured to use timecode but none is available on the engine."), *OpenUrl);
+			UE_LOG(LogMediaIOCore, Warning, TEXT("The video '%s' is configured to use timecode but none is available on the engine."), *OpenUrl);
 		}
 	}
 	else
@@ -291,10 +302,10 @@ void FMediaIOCorePlayerBase::TickTimeManagement()
 void FMediaIOCorePlayerBase::TickFetch(FTimespan DeltaTime, FTimespan Timecode)
 {
 	// Running JITR?
-	if (CurrentState == EMediaState::Playing && IsJustInTimeRenderingEnabled())
+	if (CurrentState == EMediaState::Playing)
 	{
-		// Nothing to do if no samples received yet
-		if (Samples->NumVideoSamples() > 0)
+		// Nothing to do if no samples received yet or if the sample is expected to be recieved later.
+		if (Samples->NumVideoSamples() > 0 || IsJustInTimeRenderingEnabled())
 		{
 			// Create new JITR proxy sample
 			JITRSamples->ProxySample = AcquireJITRProxySampleInitialized();
@@ -1132,7 +1143,7 @@ void FMediaIOCorePlayerBase::AddVideoSample(const TSharedRef<FMediaIOCoreTexture
 	}
 }
 
-bool FMediaIOCorePlayerBase::JustInTimeSampleRender_RenderThread(FRHICommandListImmediate& RHICmdList, TSharedPtr<FMediaIOCoreTextureSampleBase>& JITRProxySample)
+bool FMediaIOCorePlayerBase::JustInTimeSampleRender_RenderThread(FRHICommandListImmediate& RHICmdList, FTextureRHIRef& InDestinationTexture, TSharedPtr<FMediaIOCoreTextureSampleBase>& JITRProxySample)
 {
 	checkSlow(IsInRenderingThread());
 
@@ -1150,7 +1161,13 @@ bool FMediaIOCorePlayerBase::JustInTimeSampleRender_RenderThread(FRHICommandList
 	}
 
 	// Pick a sample to render
-	TSharedPtr<FMediaIOCoreTextureSampleBase> SourceSample = PickSampleToRender_RenderThread(JITRProxySample);
+	FFrameInfo FrameInformation;
+	FrameInformation.RequestedTimecode = JITRProxySample->GetTimecode().Get(FTimecode());
+	FrameInformation.SampleTimespan = JITRProxySample->GetTime().Time;
+	FrameInformation.EvaluationOffset = JITRProxySample->GetEvaluationOffsetInSeconds();
+
+	TSharedPtr<FMediaIOCoreTextureSampleBase> SourceSample = PickSampleToRender_RenderThread(FrameInformation);
+
 	if (!SourceSample)
 	{
 		UE_LOG(LogMediaIOCore, Warning, TEXT("JustInTimeSampleRender couldn't find a sample to render"));
@@ -1180,13 +1197,24 @@ TSharedPtr<FMediaIOCoreTextureSampleBase> FMediaIOCorePlayerBase::AcquireJITRPro
 {
 	// Create a new media sample acting as a dummy container to be picked by MFW which we will fill during the late update
 	FMediaIOCoreSampleJITRConfigurationArgs Args;
-	Args.Width     = 1; // Dummy value, will be replaced on rendering
-	Args.Height    = 1; // Dummy value, will be replaced on rendering
+	Args.Width     = VideoTrackFormat.Dim.X; 
+	Args.Height    = VideoTrackFormat.Dim.Y;
 	Args.Player    = AsShared().ToSharedPtr();
-	Args.Time      = FTimespan::FromSeconds(GetPlatformSeconds());
-	Args.Timecode  = FApp::GetTimecode();
+	Args.Timecode = FApp::GetTimecode();
+
+	TOptional<FQualifiedFrameTime> CurrentFrameTime = FApp::GetCurrentFrameTime();
+	if (CurrentFrameTime.IsSet() && EvaluationType != EMediaIOSampleEvaluationType::Latest)
+	{
+		Args.Time = FTimespan::FromSeconds(CurrentFrameTime.GetValue().AsSeconds());
+	}
+	else
+	{
+		Args.Time = FTimespan::FromSeconds(GetPlatformSeconds());
+	}
+
 	Args.EvaluationOffsetInSeconds = GetEvaluationOffsetInSeconds();
 	Args.Converter = CreateTextureSampleConverter();
+	Args.FrameRate = VideoFrameRate;
 
 	check(Args.Converter.IsValid());
 
@@ -1210,30 +1238,30 @@ TSharedPtr<FMediaIOCoreTextureSampleBase> FMediaIOCorePlayerBase::AcquireJITRPro
 	return NewSample;
 }
 
-TSharedPtr<FMediaIOCoreTextureSampleBase> FMediaIOCorePlayerBase::PickSampleToRender_RenderThread(const TSharedPtr<FMediaIOCoreTextureSampleBase>& JITRProxySample)
+TSharedPtr<FMediaIOCoreTextureSampleBase> FMediaIOCorePlayerBase::PickSampleToRender_RenderThread(const FFrameInfo& InFrameInformation)
 {
 	switch (EvaluationType)
 	{
 	// Pick the latest available sample
 	case EMediaIOSampleEvaluationType::Latest:
-		return PickSampleToRenderForLatest_RenderThread(JITRProxySample);
+		return PickSampleToRenderForLatest_RenderThread(InFrameInformation);
 
 	// Pick a sample based on the platform time
 	case EMediaIOSampleEvaluationType::PlatformTime:
-		return PickSampleToRenderForTimeSynchronized_RenderThread(JITRProxySample);
+		return PickSampleToRenderForTimeSynchronized_RenderThread(InFrameInformation);
 
 	// Pick a sample based on the engine's timecode
 	case EMediaIOSampleEvaluationType::Timecode:
 		return (bFramelock ?
-			PickSampleToRenderFramelocked_RenderThread(JITRProxySample) :
-			PickSampleToRenderForTimeSynchronized_RenderThread(JITRProxySample));
+			PickSampleToRenderFramelocked_RenderThread(InFrameInformation) :
+			PickSampleToRenderForTimeSynchronized_RenderThread(InFrameInformation));
 	default:
 		checkNoEntry();
-		return PickSampleToRenderForLatest_RenderThread(JITRProxySample);
+		return PickSampleToRenderForLatest_RenderThread(InFrameInformation);
 	}
 }
 
-TSharedPtr<FMediaIOCoreTextureSampleBase> FMediaIOCorePlayerBase::PickSampleToRenderForLatest_RenderThread(const TSharedPtr<FMediaIOCoreTextureSampleBase>& JITRProxySample)
+TSharedPtr<FMediaIOCoreTextureSampleBase> FMediaIOCorePlayerBase::PickSampleToRenderForLatest_RenderThread(const FFrameInfo& InFrameInformation)
 {
 	const TArray<TSharedPtr<IMediaTextureSample>> TextureSamples = Samples->GetVideoSamples();
 
@@ -1243,28 +1271,26 @@ TSharedPtr<FMediaIOCoreTextureSampleBase> FMediaIOCorePlayerBase::PickSampleToRe
 		nullptr;
 }
 
-TSharedPtr<FMediaIOCoreTextureSampleBase> FMediaIOCorePlayerBase::PickSampleToRenderForTimeSynchronized_RenderThread(const TSharedPtr<FMediaIOCoreTextureSampleBase>& JITRProxySample)
+TSharedPtr<FMediaIOCoreTextureSampleBase> FMediaIOCorePlayerBase::PickSampleToRenderForTimeSynchronized_RenderThread(const FFrameInfo& InFrameInformation)
 {
 	// Reference time based on evaluation type
 	FTimespan TargetSampleTimespan;
 
 	// Get base uncorrected reference point
-	const FTimecode InvalidTimecode;
-	const FTimecode RequestedTimecode = JITRProxySample->GetTimecode().Get(InvalidTimecode);
-	if (EvaluationType == EMediaIOSampleEvaluationType::Timecode && RequestedTimecode != FTimecode())
+	
+	if (EvaluationType == EMediaIOSampleEvaluationType::Timecode && InFrameInformation.RequestedTimecode != FTimecode())
 	{
 		// We'll use timecode data to find a proper sample
-		TargetSampleTimespan = RequestedTimecode.ToTimespan(VideoFrameRate);
+		TargetSampleTimespan = InFrameInformation.RequestedTimecode.ToTimespan(VideoFrameRate);
 	}
 	else
 	{
 		// We'll use platform time to find a proper sample
-		TargetSampleTimespan = JITRProxySample->GetTime().Time;
+		TargetSampleTimespan = InFrameInformation.SampleTimespan;
 	}
 
 	// Apply time correction to the target time
-	const double RequestedOffsetInSeconds = JITRProxySample->GetEvaluationOffsetInSeconds();
-	const FTimespan RequestedOffsetTimespan = FTimespan::FromSeconds(RequestedOffsetInSeconds);
+	const FTimespan RequestedOffsetTimespan = FTimespan::FromSeconds(InFrameInformation.EvaluationOffset);
 	const FTimespan TargetTimespanCorrected = TargetSampleTimespan + RequestedOffsetTimespan;
 
 	// Go over the sample pool and find a sample closest to the target time
@@ -1273,7 +1299,10 @@ TSharedPtr<FMediaIOCoreTextureSampleBase> FMediaIOCorePlayerBase::PickSampleToRe
 
 	// Get all available video samples
 	const TArray<TSharedPtr<IMediaTextureSample>> TextureSamples = Samples->GetVideoSamples();
-
+	if (TextureSamples.Num() == 0)
+	{
+		return nullptr;
+	}
 	for (int32 Index = 0; Index < TextureSamples.Num(); ++Index)
 	{
 		// When EvaluationType == ETimedDataInputEvaluationType::Timecode, the time represents the sample's timecode.
@@ -1299,15 +1328,40 @@ TSharedPtr<FMediaIOCoreTextureSampleBase> FMediaIOCorePlayerBase::PickSampleToRe
 	}
 
 	checkSlow(ClosestIndex >= 0 && ClosestIndex < TextureSamples.Num());
-
 	// Finally, return the closest sample we found
 	return StaticCastSharedPtr<FMediaIOCoreTextureSampleBase, IMediaTextureSample, ESPMode::ThreadSafe>(TextureSamples[ClosestIndex]);
 }
 
-TSharedPtr<FMediaIOCoreTextureSampleBase> FMediaIOCorePlayerBase::PickSampleToRenderFramelocked_RenderThread(const TSharedPtr<FMediaIOCoreTextureSampleBase>& JITRProxySample)
+TSharedPtr<FMediaIOCoreTextureSampleBase> FMediaIOCorePlayerBase::PickSampleToRenderFramelocked_RenderThread(const FFrameInfo& InFrameInformation)
 {
-	unimplemented();
-	return PickSampleToRenderForTimeSynchronized_RenderThread(JITRProxySample);
+	// This isn't an ideal version of frame locking as this will block the render thread of the receiving node to wait for frame from the sending device
+	uint32 CurrentLatency = FrameDelay;
+	TSharedPtr<FMediaIOCoreTextureSampleBase> SampleToReturn;
+	constexpr double TimeoutSeconds = 1.0;
+	const double StartTimeSeconds = FPlatformTime::Seconds();
+	uint64 FrameNumber = GFrameCounterRenderThread + CurrentLatency - 1;
+	while (true)
+	{
+		const TArray<TSharedPtr<IMediaTextureSample>> TextureSamples = Samples->GetVideoSamples();
+		for (int32 Index = 0; Index < TextureSamples.Num(); ++Index)
+		{
+			SampleToReturn = StaticCastSharedPtr<FMediaIOCoreTextureSampleBase>(TextureSamples[Index]);
+			if (SampleToReturn->GetFrameNumber() == FrameNumber)
+			{
+				return SampleToReturn;
+			}
+		}
+		SampleToReturn = nullptr;
+		FPlatformProcess::SleepNoStats(0.001);
+		if ((FPlatformTime::Seconds() - StartTimeSeconds) > TimeoutSeconds)
+		{
+			break;
+		}
+	}
+
+	UE_LOG(LogMediaIOCore, Warning, TEXT("Timeout waiting for frame #%u to be rendered."), FrameNumber);
+	// Just return the most recent sample available in the queue
+	return nullptr;
 }
 
 void FMediaIOCorePlayerBase::TransferTexture_RenderThread(FRHICommandListImmediate& RHICmdList, const TSharedPtr<FMediaIOCoreTextureSampleBase>& Sample, const TSharedPtr<FMediaIOCoreTextureSampleBase>& JITRProxySample)
@@ -1334,6 +1388,26 @@ void FMediaIOCorePlayerBase::TransferTexture_RenderThread(FRHICommandListImmedia
 	}
 }
 
+void FMediaIOCorePlayerBase::LogBookmark(const FString& Text, const TSharedRef<IMediaTextureSample>& Sample)
+{
+	// Put a bookmark if requested
+	const bool bDetailedInsights = MediaIOCorePlayerDetail::CVarMediaIOJITRInsights.GetValueOnAnyThread();
+	if (bDetailedInsights)
+	{
+		const TOptional<FTimecode> Timecode = Sample->GetTimecode();
+		if (Timecode.IsSet())
+		{
+			const FFrameNumber FrameNumber = Timecode.GetValue().ToFrameNumber(VideoFrameRate);
+			TRACE_BOOKMARK(TEXT("%s [%d]"), *Text, FrameNumber.Value % 100);
+		}
+		else
+		{
+			const int32 Milliseconds = Sample->GetTime().Time.GetFractionMilli();
+			TRACE_BOOKMARK(TEXT("%s [%d]"), *Text, Milliseconds);
+		}
+	}
+}
+
 bool FMediaIOCorePlayerBase::FJITRMediaTextureSamples::FetchVideo(TRange<FTimespan> TimeRange, TSharedPtr<IMediaTextureSample>& OutSample)
 {
 	if (!ProxySample.IsValid())
@@ -1358,24 +1432,5 @@ void FMediaIOCorePlayerBase::FJITRMediaTextureSamples::FlushSamples()
 	ProxySample.Reset();
 }
 
-void FMediaIOCorePlayerBase::LogBookmark(const FString& Text, const TSharedRef<IMediaTextureSample>& Sample)
-{
-	// Put a bookmark if requested
-	const bool bDetailedInsights = MediaIOCorePlayerDetail::CVarMediaIOJITRInsights.GetValueOnAnyThread();
-	if (bDetailedInsights)
-	{
-		const TOptional<FTimecode> Timecode = Sample->GetTimecode();
-		if (Timecode.IsSet())
-		{
-			const FFrameNumber FrameNumber = Timecode.GetValue().ToFrameNumber(VideoFrameRate);
-			TRACE_BOOKMARK(TEXT("%s [%d]"), *Text, FrameNumber.Value % 100);
-		}
-		else
-		{
-			const int32 Milliseconds = Sample->GetTime().Time.GetFractionMilli();
-			TRACE_BOOKMARK(TEXT("%s [%d]"), *Text, Milliseconds);
-		}
-	}
-}
 
 #undef LOCTEXT_NAMESPACE
diff --git a/Engine/Plugins/Media/MediaIOFramework/Source/MediaIOCore/Private/Player/MediaIOCoreSamples.cpp b/Engine/Plugins/Media/MediaIOFramework/Source/MediaIOCore/Private/Player/MediaIOCoreSamples.cpp
--- a/Engine/Plugins/Media/MediaIOFramework/Source/MediaIOCore/Private/Player/MediaIOCoreSamples.cpp
+++ b/Engine/Plugins/Media/MediaIOFramework/Source/MediaIOCore/Private/Player/MediaIOCoreSamples.cpp
@@ -104,9 +104,17 @@ void FMediaIOCoreSamples::FlushSamples()
 	VideoSamples.FlushSamples();
 }
 
-
 bool FMediaIOCoreSamples::PeekVideoSampleTime(FMediaTimeStamp & TimeStamp)
 {
 	// player does not support v2 timing control at this point -> no need for this method, yet
 	return false;
-}
\ No newline at end of file
+}
+
+void FMediaIOCoreSamples::SetSampleBufferSize(int32 BufferSize)
+{
+	AudioSamples.SetDataBufferSize(BufferSize);
+	CaptionSamples.SetDataBufferSize(BufferSize);
+	MetadataSamples.SetDataBufferSize(BufferSize);
+	SubtitleSamples.SetDataBufferSize(BufferSize);
+	VideoSamples.SetDataBufferSize(BufferSize);
+}
diff --git a/Engine/Plugins/Media/MediaIOFramework/Source/MediaIOCore/Private/Shared/MediaIOCoreTextureSampleBase.cpp b/Engine/Plugins/Media/MediaIOFramework/Source/MediaIOCore/Private/Shared/MediaIOCoreTextureSampleBase.cpp
--- a/Engine/Plugins/Media/MediaIOFramework/Source/MediaIOCore/Private/Shared/MediaIOCoreTextureSampleBase.cpp
+++ b/Engine/Plugins/Media/MediaIOFramework/Source/MediaIOCore/Private/Shared/MediaIOCoreTextureSampleBase.cpp
@@ -2,6 +2,7 @@
 
 #include "MediaIOCoreTextureSampleBase.h"
 
+#include "Async/Async.h"
 #include "MediaIOCoreTextureSampleConverter.h"
 
 #include "OpenColorIOConfiguration.h"
@@ -13,7 +14,6 @@
 #include "RHI.h"
 #include "RHIResources.h"
 #include "ScreenPass.h"
-
 DECLARE_GPU_STAT(MediaIO_ColorConversion);
 
 FMediaIOCoreTextureSampleBase::FMediaIOCoreTextureSampleBase()
@@ -150,6 +150,15 @@ bool FMediaIOCoreTextureSampleBase::SetBufferWithEvenOddLine(bool bUseEvenLine,
 	return true;
 }
 
+void FMediaIOCoreTextureSampleBase::SetColorConversionSettings(TSharedPtr<struct FOpenColorIOColorConversionSettings> InColorConversionSettings)
+{
+	ColorConversionSettings = InColorConversionSettings;
+
+	AsyncTask(ENamedThreads::GameThread, [this]() {
+		CacheColorCoversionSettings_GameThread();
+	});
+}
+
 void* FMediaIOCoreTextureSampleBase::RequestBuffer(uint32 InBufferSize)
 {
 	FreeSample();
@@ -157,6 +166,16 @@ void* FMediaIOCoreTextureSampleBase::RequestBuffer(uint32 InBufferSize)
 	return Buffer.GetData();
 }
 
+void* FMediaIOCoreTextureSampleBase::GetOrRequestBuffer(uint32 InBufferSize)
+{
+	if (Buffer.Num() != InBufferSize)
+	{
+		RequestBuffer(InBufferSize);
+	}
+
+	return Buffer.GetData();
+}
+
 TSharedPtr<FMediaIOCorePlayerBase> FMediaIOCoreTextureSampleBase::GetPlayer() const
 {
 	return Player.Pin();
@@ -174,6 +193,8 @@ bool FMediaIOCoreTextureSampleBase::InitializeJITR(const FMediaIOCoreSampleJITRC
 	Height = Args.Height;
 	Time   = Args.Time;
 	Timecode = Args.Timecode;
+	FrameNumber = GFrameNumber;
+	Duration = FTimespan(ETimespan::TicksPerSecond * Args.FrameRate.AsInterval());
 
 	// JITR data
 	Player    = Args.Player;
@@ -207,44 +228,20 @@ void FMediaIOCoreTextureSampleBase::CopyConfiguration(const TSharedPtr<FMediaIOC
 	OriginalSample = SourceSample;
 }
 
-bool FMediaIOCoreTextureSampleBase::ApplyColorConversion(FRHICommandListImmediate& RHICmdList, FTextureRHIRef& InSrcTexture, FTextureRHIRef& InDstTexture)
+void FMediaIOCoreTextureSampleBase::CacheColorCoversionSettings_GameThread()
 {
-	if (!CachedOCIOResources.IsValid())
+	if (ColorConversionSettings.IsValid() && ColorConversionSettings->IsValid())
 	{
-		if (ColorConversionSettings.IsValid() && ColorConversionSettings->IsValid())
-		{
-			CachedOCIOResources = MakeShared<FOpenColorIORenderPassResources>();
-
-			FOpenColorIOTransformResource* ShaderResource = nullptr;
-			TSortedMap<int32, TWeakObjectPtr<UTexture>> TransformTextureResources;
-
-			if (ColorConversionSettings->ConfigurationSource != nullptr)
-			{
-				const bool bFoundTransform = ColorConversionSettings->ConfigurationSource->GetRenderResources(
-					GMaxRHIFeatureLevel
-					, *ColorConversionSettings
-					, ShaderResource
-					, TransformTextureResources);
-
-				if (bFoundTransform)
-				{
-					// Transform was found, so shader must be there but doesn't mean the actual shader is available
-					check(ShaderResource);
-					if (ShaderResource->GetShader<FOpenColorIOPixelShader>().IsNull())
-					{
-						ensureMsgf(false, TEXT("Can't apply display look - Shader was invalid for Resource %s"), *ShaderResource->GetFriendlyName());
-
-						//Invalidate shader resource
-						ShaderResource = nullptr;
-					}
-				}
-			}
-
-			CachedOCIOResources->ShaderResource = ShaderResource;
-			CachedOCIOResources->TextureResources = TransformTextureResources;
-		}
+		CachedOCIOResources = MakeShared<FOpenColorIORenderPassResources>();
+
+		FOpenColorIORenderPassResources Resources = FOpenColorIORendering::GetRenderPassResources(*ColorConversionSettings, GMaxRHIFeatureLevel);
+		CachedOCIOResources->ShaderResource = Resources.ShaderResource;
+		CachedOCIOResources->TextureResources = Resources.TextureResources;
 	}
+}
 
+bool FMediaIOCoreTextureSampleBase::ApplyColorConversion(FRHICommandListImmediate& RHICmdList, FTextureRHIRef& InSrcTexture, FTextureRHIRef& InDstTexture)
+{
 	if (CachedOCIOResources)
 	{
 		FRDGBuilder GraphBuilder(RHICmdList);
@@ -381,4 +378,4 @@ float FMediaIOCoreTextureSampleBase::GetHDRNitsNormalizationFactor() const
     }
 	
 	return (GetEncodingType() == UE::Color::EEncoding::sRGB || GetEncodingType() == UE::Color::EEncoding::Linear) ? 1.0f : kMediaSample_HDR_NitsNormalizationFactor;
-}
+}
\ No newline at end of file
diff --git a/Engine/Plugins/Media/MediaIOFramework/Source/MediaIOCore/Public/MediaIOCorePlayerBase.h b/Engine/Plugins/Media/MediaIOFramework/Source/MediaIOCore/Public/MediaIOCorePlayerBase.h
--- a/Engine/Plugins/Media/MediaIOFramework/Source/MediaIOCore/Public/MediaIOCorePlayerBase.h
+++ b/Engine/Plugins/Media/MediaIOFramework/Source/MediaIOCore/Public/MediaIOCorePlayerBase.h
@@ -4,6 +4,7 @@
 
 #include "CoreMinimal.h"
 
+#include "Engine/TextureDefines.h"
 #include "IMediaCache.h"
 #include "IMediaControls.h"
 #include "IMediaPlayer.h"
@@ -162,12 +163,23 @@ public:
 	virtual bool SupportsSubFrames() const override;
 
 public:
+	class FFrameInfo
+	{
+	public:
+		FTimecode RequestedTimecode = FTimecode();
+		FTimespan SampleTimespan = FTimespan();
+		double EvaluationOffset = 0.;
+		uint32 FrameNumber = 0.;
+	};
 
+	/** Deprecated in UE5.5. */
+	bool JustInTimeSampleRender_RenderThread(FRHICommandListImmediate& RHICmdList, TSharedPtr<FMediaIOCoreTextureSampleBase>& JITRProxySample) { return false; };
+	
 	/**
 	 * Just in time sample rendering. This method is responsible for late sample picking,
 	 * then rendering it into the proxy sample provided.
 	 */
-	bool JustInTimeSampleRender_RenderThread(FRHICommandListImmediate& RHICmdList, TSharedPtr<FMediaIOCoreTextureSampleBase>& JITRProxySample);
+	virtual bool JustInTimeSampleRender_RenderThread(FRHICommandListImmediate& RHICmdList, FTextureRHIRef& InDestinationTexture, TSharedPtr<FMediaIOCoreTextureSampleBase>& JITRProxySample);
 
 protected:
 
@@ -254,10 +266,10 @@ protected:
 	virtual TSharedPtr<FMediaIOCoreTextureSampleConverter> CreateTextureSampleConverter() const;
 
 	/** Pick a sample to render */
-	TSharedPtr<FMediaIOCoreTextureSampleBase> PickSampleToRender_RenderThread(const TSharedPtr<FMediaIOCoreTextureSampleBase>& JITRProxySample);
-	TSharedPtr<FMediaIOCoreTextureSampleBase> PickSampleToRenderForLatest_RenderThread(const TSharedPtr<FMediaIOCoreTextureSampleBase>& JITRProxySample);
-	TSharedPtr<FMediaIOCoreTextureSampleBase> PickSampleToRenderForTimeSynchronized_RenderThread(const TSharedPtr<FMediaIOCoreTextureSampleBase>& JITRProxySample);
-	TSharedPtr<FMediaIOCoreTextureSampleBase> PickSampleToRenderFramelocked_RenderThread(const TSharedPtr<FMediaIOCoreTextureSampleBase>& JITRProxySample);
+	virtual TSharedPtr<FMediaIOCoreTextureSampleBase> PickSampleToRender_RenderThread(const FFrameInfo& InFrameInformation);
+	virtual TSharedPtr<FMediaIOCoreTextureSampleBase> PickSampleToRenderForLatest_RenderThread(const FFrameInfo& InFrameInformation);
+	virtual TSharedPtr<FMediaIOCoreTextureSampleBase> PickSampleToRenderForTimeSynchronized_RenderThread(const FFrameInfo& InFrameInformation);
+	virtual TSharedPtr<FMediaIOCoreTextureSampleBase> PickSampleToRenderFramelocked_RenderThread(const FFrameInfo& InFrameInformation);
 
 	/**
 	 * A wrapper method responsible for transferring of the sample textures into GPU memory based
@@ -317,6 +329,24 @@ protected:
 	/** Open color IO conversion data. */
 	TSharedPtr<struct FOpenColorIOColorConversionSettings> OCIOSettings;
 
+	/** Whether media playback should be framelocked to the engine's timecode */
+	bool bFramelock = false;
+
+	/** Used to ensure that JIT rendering is executed only once per frame */
+	uint64 LastEngineRTFrameThatUpdatedJustInTime = TNumericLimits<uint64>::Max();
+
+	/** Whether to override the source encoding or to use the metadata embedded in the ancillary data of the signal. */
+	bool bOverrideSourceEncoding = true;
+
+	/** Encoding of the source texture. */
+	ETextureSourceEncoding OverrideSourceEncoding = ETextureSourceEncoding::TSE_Linear;
+
+	/** Whether to override the source color space or to use the metadata embedded in the ancillary data of the signal. */
+	bool bOverrideSourceColorSpace = true;
+
+	/** Color space of the source texture. */
+	ETextureColorSpace OverrideSourceColorSpace = ETextureColorSpace::TCS_None;
+
 private:
 	void OnSampleDestroyed(TRefCountPtr<FRHITexture> InTexture);
 	void RegisterSampleBuffer(const TSharedPtr<FMediaIOCoreTextureSampleBase>& InSample);
@@ -378,12 +408,6 @@ private:
 	/** Is Just-In-Time Rendering enabled */
 	bool bJustInTimeRender = false;
 
-	/** Whether media playback should be framelocked to the engine's timecode */
-	bool bFramelock = false;
-
 	/** JITR samples proxy */
 	const TUniquePtr<FJITRMediaTextureSamples> JITRSamples;
-
-	/** Used to ensure that JIT rendering is executed only once per frame */
-	uint64 LastEngineRTFrameThatUpdatedJustInTime = TNumericLimits<uint64>::Max();
 };
diff --git a/Engine/Plugins/Media/MediaIOFramework/Source/MediaIOCore/Public/MediaIOCoreSamples.h b/Engine/Plugins/Media/MediaIOFramework/Source/MediaIOCore/Public/MediaIOCoreSamples.h
--- a/Engine/Plugins/Media/MediaIOFramework/Source/MediaIOCore/Public/MediaIOCoreSamples.h
+++ b/Engine/Plugins/Media/MediaIOFramework/Source/MediaIOCore/Public/MediaIOCoreSamples.h
@@ -306,6 +306,7 @@ public:
 	virtual bool FetchVideo(TRange<FTimespan> TimeRange, TSharedPtr<IMediaTextureSample, ESPMode::ThreadSafe>& OutSample) override;
 	virtual void FlushSamples() override;
 	virtual bool PeekVideoSampleTime(FMediaTimeStamp & TimeStamp) override;
+	virtual void SetSampleBufferSize(int32 BufferSize) override;
 
 protected:
 
diff --git a/Engine/Plugins/Media/MediaIOFramework/Source/MediaIOCore/Public/MediaIOCoreTextureSampleBase.h b/Engine/Plugins/Media/MediaIOFramework/Source/MediaIOCore/Public/MediaIOCoreTextureSampleBase.h
--- a/Engine/Plugins/Media/MediaIOFramework/Source/MediaIOCore/Public/MediaIOCoreTextureSampleBase.h
+++ b/Engine/Plugins/Media/MediaIOFramework/Source/MediaIOCore/Public/MediaIOCoreTextureSampleBase.h
@@ -77,6 +77,9 @@ struct MEDIAIOCORE_API FMediaIOCoreSampleJITRConfigurationArgs
 
 	/** Sample converter to process this sample */
 	TSharedPtr<FMediaIOCoreTextureSampleConverter> Converter;
+	
+	/** Frame rate of the current sample. */
+	FFrameRate FrameRate;
 };
 
 
@@ -211,10 +214,8 @@ public:
 	/**
 	 * Set the OCIO settings used for color conversion.
 	 */
-	void SetColorConversionSettings(TSharedPtr<struct FOpenColorIOColorConversionSettings> InColorConversionSettings)
-	{
-		ColorConversionSettings = InColorConversionSettings;
-	}
+	void SetColorConversionSettings(TSharedPtr<struct FOpenColorIOColorConversionSettings> InColorConversionSettings);
+
 	/**
 	 * Request an uninitialized sample buffer.
 	 * Should be used when the buffer could be filled by something else.
@@ -224,6 +225,7 @@ public:
 	 */
 	virtual void* RequestBuffer(uint32 InBufferSize);
 
+
 	/**
 	 * Configure this sample for JITR
 	 *
@@ -232,9 +234,9 @@ public:
 	virtual bool InitializeJITR(const FMediaIOCoreSampleJITRConfigurationArgs& Args);
 
 	/** Marks this sample as one that is ready and awaiting for fast GPUDirect texture transfer */
-	void SetAwaitingForGPUTransfer()
+	void SetAwaitingForGPUTransfer(bool bIsAwaitingGPUTransfer = true)
 	{
-		bIsAwaitingForGPUTransfer = true;
+		bIsAwaitingForGPUTransfer = bIsAwaitingGPUTransfer;
 	}
 
 	/** Returns whether it's ready for GPUDirect texture transfer */
@@ -254,6 +256,11 @@ public:
 
 	/** Copies all neccessary data from a source sample to render JIT */
 	virtual void CopyConfiguration(const TSharedPtr<FMediaIOCoreTextureSampleBase>& SourceSample);
+protected:
+	/**
+	* Method that caches color conversion settings on Game thread.
+	*/
+	void CacheColorCoversionSettings_GameThread();
 
 public:
 	//~ IMediaTextureSample interface
@@ -339,6 +346,30 @@ public:
 		return Buffer.GetData();
 	}
 
+	/**
+	 * Attemps to get the initialized buffer.
+	 * If buffer isn't initialized or of a different size, requests a new one.
+	 *
+	 * @param InBufferSize The size of the required video buffer.
+	 */
+	virtual void* GetOrRequestBuffer(uint32 InBufferSize);
+
+
+	virtual uint32 GetFrameNumber() const
+	{
+		return FrameNumber;
+	}
+
+	virtual void SetTime(const FTimespan& InTime)
+	{
+		Time = InTime;
+	}
+
+	virtual void SetFrameNumber(uint32 InFrameNumber)
+	{
+		FrameNumber = InFrameNumber;
+	}
+
 	//~ IMediaTextureSampleColorConverter interface
 	virtual bool ApplyColorConversion(FRHICommandListImmediate& RHICmdList, FTextureRHIRef& InSrcTexture, FTextureRHIRef& InDstTexture) override;
 
@@ -409,6 +440,9 @@ protected:
 	/** Sample timecode. */
 	TOptional<FTimecode> Timecode;
 
+	/** Which engine frame number this sample corresponds to. */
+	uint32 FrameNumber = 0;
+
 	/** Image dimensions */
 	uint32 Stride = 0;
 	uint32 Width  = 0;
@@ -433,13 +467,13 @@ protected:
 	/** Color space structure of the incoming texture. Used for retrieving chromaticities. */
 	UE::Color::FColorSpace ColorSpaceStruct = UE::Color::FColorSpace(UE::Color::EColorSpace::sRGB);
 
-private:
 	/** The player that created this sample */
 	TWeakPtr<FMediaIOCorePlayerBase> Player;
 
 	/** Custom converter that will be the one checking back with the player for just in time sample render purposes */
 	TSharedPtr<FMediaIOCoreTextureSampleConverter> Converter;
 
+private:
 	/**
 	 * A reference to the original sample that was chosen during JITR. The idea of this member is to keep
 	 * the original sample alive, prevent any of its resources from being released while this proxy sample is in use.
diff --git a/Engine/Source/Runtime/Media/Public/IMediaSamples.h b/Engine/Source/Runtime/Media/Public/IMediaSamples.h
--- a/Engine/Source/Runtime/Media/Public/IMediaSamples.h
+++ b/Engine/Source/Runtime/Media/Public/IMediaSamples.h
@@ -116,6 +116,9 @@ public:
 		// override in child classes, if supported
 	}
 
+	/** Sets the number of samples to be stored in sample container. */
+	virtual void SetSampleBufferSize(int32 BufferSize) {};
+
 	enum class EFetchBestSampleResult
 	{
 		Ok = 0,
diff --git a/Engine/Source/Runtime/MediaUtils/Public/MediaObjectPool.h b/Engine/Source/Runtime/MediaUtils/Public/MediaObjectPool.h
--- a/Engine/Source/Runtime/MediaUtils/Public/MediaObjectPool.h
+++ b/Engine/Source/Runtime/MediaUtils/Public/MediaObjectPool.h
@@ -92,8 +92,11 @@ class TMediaObjectPool
 
 	public:
 
-		/** Acquire an object from the pool. */
-		ObjectType* Acquire()
+		/** 
+		* Acquire an object from the pool. 
+		* bAlloc allows to chose if new object should be allocated and added to the pool.
+		*/
+		ObjectType* Acquire(bool bAlloc = true)
 		{
 			ObjectType* Result = nullptr;
 			{
@@ -129,10 +132,14 @@ class TMediaObjectPool
 				}
 			}
 
-			if (Result == nullptr)
+			if (Result == nullptr && bAlloc)
 			{
 				Result = ObjectAllocatorInstance->Alloc();
 			}
+			else
+			{
+				return Result;
+			}
 			
 			Result->InitializePoolable();
 
@@ -281,13 +288,14 @@ public:
 	 * Use the Release method to return the object to the pool.
 	 * You can use the ToShared and ToUnique methods to convert
 	 * this object to a tracked shared object later if desired.
+	 * bAlloc allows to chose if new object should be allocated and added to the pool.
 	 *
 	 * @return The object.
 	 * @see AcquireShared, AcquireUnique, Release, ToShared, ToUnique
 	 */
-	ObjectType* Acquire()
+	ObjectType* Acquire(bool bAlloc = true)
 	{
-		return Storage->Acquire();
+		return Storage->Acquire(bAlloc);
 	}
 
 	/**
@@ -301,7 +309,28 @@ public:
 	 */
 	TSharedRef<ObjectType, ESPMode::ThreadSafe> AcquireShared()
 	{
-		ObjectType* Object = Acquire();
+		ObjectType* Object = Acquire(true);
+		check(Object != nullptr);
+		return MakeShareable(Object, TDeleter(Storage));
+	}
+
+	/**
+	 * Acquire a shared object from the pool.
+	 *
+	 * Shared objects do not need to be returned to the pool. They'll be
+	 * reclaimed automatically when their reference count goes to zero.
+	 * bAlloc allows to chose if new object should be allocated and added to the pool.
+	 *
+	 * @return The shared object.
+	 * @see Acquire, AcquireUnique, Reset, ToShared
+	 */
+	TSharedPtr<ObjectType, ESPMode::ThreadSafe> AcquireShared(bool bAlloc)
+	{
+		ObjectType* Object = Acquire(bAlloc);
+		if (!Object && !bAlloc)
+		{
+			return nullptr;
+		}
 		check(Object != nullptr);
 
 		return MakeShareable(Object, TDeleter(Storage));
