Media IO and Media Texture Resource:
- Added support for color conversion for Converters that provide their own texture.
- Capped the Frame Delay to 4 frames so that the buffers don't grow out of control.
- Other minor improvements.

#rnx

diff --git a/Engine/Plugins/Media/MediaIOFramework/Source/MediaIOCore/Private/Player/MediaIOCorePlayerBase.cpp b/Engine/Plugins/Media/MediaIOFramework/Source/MediaIOCore/Private/Player/MediaIOCorePlayerBase.cpp
--- a/Engine/Plugins/Media/MediaIOFramework/Source/MediaIOCore/Private/Player/MediaIOCorePlayerBase.cpp
+++ b/Engine/Plugins/Media/MediaIOFramework/Source/MediaIOCore/Private/Player/MediaIOCorePlayerBase.cpp
@@ -304,7 +304,11 @@ void FMediaIOCorePlayerBase::TickFetch(FTimespan DeltaTime, FTimespan Timecode)
 	// Running JITR?
 	if (CurrentState == EMediaState::Playing)
 	{
-		// Nothing to do if no samples received yet or if the sample is expected to be recieved later.
+		// If not in Just In Time Mode and no samples received yet, don't provide the ProxySample otherwise rendering will be scheduled.
+		// Otherwise just in time implies that the sample most likely is going to arrive closer to the render thread / RHI thread command execution.
+		// AcquireJITRProxySampleInitialized is called on game thread before any rendering commands are executed for this frame.
+		// With just in time enabled, the Sample list could be empty and the sample could arrive much later. 
+		// The inheriting players must be responsible for handling the provision of JIT samples at the later stage.
 		if (Samples->NumVideoSamples() > 0 || IsJustInTimeRenderingEnabled())
 		{
 			// Create new JITR proxy sample
@@ -1230,10 +1234,12 @@ TSharedPtr<FMediaIOCoreTextureSampleBase> FMediaIOCorePlayerBase::AcquireJITRPro
 
 	// Initialize the JITR sample
 	const bool bInitialized = NewSample->InitializeJITR(Args);
+
 	if (!bInitialized)
 	{
 		return nullptr;
 	}
+	NewSample->SetColorConversionSettings(OCIOSettings);
 
 	return NewSample;
 }
diff --git a/Engine/Plugins/Media/MediaIOFramework/Source/MediaIOCore/Private/Shared/MediaIOCoreTextureSampleBase.cpp b/Engine/Plugins/Media/MediaIOFramework/Source/MediaIOCore/Private/Shared/MediaIOCoreTextureSampleBase.cpp
--- a/Engine/Plugins/Media/MediaIOFramework/Source/MediaIOCore/Private/Shared/MediaIOCoreTextureSampleBase.cpp
+++ b/Engine/Plugins/Media/MediaIOFramework/Source/MediaIOCore/Private/Shared/MediaIOCoreTextureSampleBase.cpp
@@ -153,10 +153,16 @@ bool FMediaIOCoreTextureSampleBase::SetBufferWithEvenOddLine(bool bUseEvenLine,
 void FMediaIOCoreTextureSampleBase::SetColorConversionSettings(TSharedPtr<struct FOpenColorIOColorConversionSettings> InColorConversionSettings)
 {
 	ColorConversionSettings = InColorConversionSettings;
-
-	AsyncTask(ENamedThreads::GameThread, [this]() {
+	if (IsInGameThread())
+	{
 		CacheColorCoversionSettings_GameThread();
-	});
+	}
+	else
+	{
+		AsyncTask(ENamedThreads::GameThread, [this]() {
+			CacheColorCoversionSettings_GameThread();
+			});
+	}
 }
 
 void* FMediaIOCoreTextureSampleBase::RequestBuffer(uint32 InBufferSize)
@@ -200,7 +206,6 @@ bool FMediaIOCoreTextureSampleBase::InitializeJITR(const FMediaIOCoreSampleJITRC
 	Player    = Args.Player;
 	Converter = Args.Converter;
 	EvaluationOffsetInSeconds = Args.EvaluationOffsetInSeconds;
-
 	return true;
 }
 
@@ -223,6 +228,13 @@ void FMediaIOCoreTextureSampleBase::CopyConfiguration(const TSharedPtr<FMediaIOC
 	ColorSpaceStruct = SourceSample->ColorSpaceStruct;
 	ColorConversionSettings = SourceSample->ColorConversionSettings;
 	CachedOCIOResources = SourceSample->CachedOCIOResources;
+	Player = SourceSample->Player;
+	Converter = SourceSample->Converter;
+	FrameNumber = SourceSample->FrameNumber.load();
+	Duration = SourceSample->Duration;
+	Texture = SourceSample->Texture;
+
+	EvaluationOffsetInSeconds = SourceSample->EvaluationOffsetInSeconds;
 
 	// Save original sample
 	OriginalSample = SourceSample;
@@ -303,6 +315,25 @@ void FMediaIOCoreTextureSampleBase::SetDestructionCallback(TFunction<void(TRefCo
 	DestructionCallback = InDestructionCallback;
 }
 
+EPixelFormat FMediaIOCoreTextureSampleBase::GetPixelFormat()
+{
+	switch (GetFormat())
+	{
+	case EMediaTextureSampleFormat::FloatRGBA:
+		return PF_FloatRGBA;
+	case EMediaTextureSampleFormat::CharBGR10A2:
+	{
+		if (GetEncodingType() != UE::Color::EEncoding::Linear)
+		{
+			return PF_FloatRGB;
+		}
+		return PF_A2B10G10R10;
+	}
+	default:
+		return PF_B8G8R8A8;
+	}
+}
+
 void FMediaIOCoreTextureSampleBase::ShutdownPoolable()
 {
 	if (DestructionCallback)
@@ -370,6 +401,15 @@ UE::Color::EEncoding FMediaIOCoreTextureSampleBase::GetEncodingType() const
 	return Encoding;
 }
 
+UE::Color::EColorSpace FMediaIOCoreTextureSampleBase::GetColorSpaceType() const
+{
+	if (ColorConversionSettings && ColorConversionSettings->IsValid())
+	{
+		return UE::Color::EColorSpace::None;
+	}
+	return ColorSpace;
+}
+
 float FMediaIOCoreTextureSampleBase::GetHDRNitsNormalizationFactor() const
 {
 	if (ColorConversionSettings && ColorConversionSettings->IsValid())
diff --git a/Engine/Plugins/Media/MediaIOFramework/Source/MediaIOCore/Public/MediaIOCorePlayerBase.h b/Engine/Plugins/Media/MediaIOFramework/Source/MediaIOCore/Public/MediaIOCorePlayerBase.h
--- a/Engine/Plugins/Media/MediaIOFramework/Source/MediaIOCore/Public/MediaIOCorePlayerBase.h
+++ b/Engine/Plugins/Media/MediaIOFramework/Source/MediaIOCore/Public/MediaIOCorePlayerBase.h
@@ -163,6 +163,7 @@ public:
 	virtual bool SupportsSubFrames() const override;
 
 public:
+	/** Class used to pass information about current frame for sample picking.*/
 	class FFrameInfo
 	{
 	public:
@@ -173,6 +174,7 @@ public:
 	};
 
 	/** Deprecated in UE5.5. */
+	UE_DEPRECATED("5.5", "Use the overloaded method that takes FFrameInfo instead of proxy sample.")
 	bool JustInTimeSampleRender_RenderThread(FRHICommandListImmediate& RHICmdList, TSharedPtr<FMediaIOCoreTextureSampleBase>& JITRProxySample) { return false; };
 	
 	/**
diff --git a/Engine/Plugins/Media/MediaIOFramework/Source/MediaIOCore/Public/MediaIOCoreTextureSampleBase.h b/Engine/Plugins/Media/MediaIOFramework/Source/MediaIOCore/Public/MediaIOCoreTextureSampleBase.h
--- a/Engine/Plugins/Media/MediaIOFramework/Source/MediaIOCore/Public/MediaIOCoreTextureSampleBase.h
+++ b/Engine/Plugins/Media/MediaIOFramework/Source/MediaIOCore/Public/MediaIOCoreTextureSampleBase.h
@@ -329,6 +329,9 @@ public:
 	virtual FVector2d GetDisplayPrimaryBlue() const override;
 	virtual UE::Color::EEncoding GetEncodingType() const override;
 	virtual float GetHDRNitsNormalizationFactor() const override;
+	
+	/**Method that returns this sample's colorspace type.*/
+	virtual UE::Color::EColorSpace GetColorSpaceType() const;
 
 	virtual const void* GetBuffer() override
 	{
@@ -355,7 +358,7 @@ public:
 	virtual void* GetOrRequestBuffer(uint32 InBufferSize);
 
 
-	virtual uint32 GetFrameNumber() const
+	virtual uint64 GetFrameNumber() const
 	{
 		return FrameNumber;
 	}
@@ -397,7 +400,7 @@ public:
 
 	void SetTexture(TRefCountPtr<FRHITexture> InRHITexture);
 	void SetDestructionCallback(TFunction<void(TRefCountPtr<FRHITexture>)> InDestructionCallback);
-
+	EPixelFormat GetPixelFormat();
 private:
 	/** Hold a texture to be used for gpu texture transfers. */
 	TRefCountPtr<FRHITexture> Texture;
@@ -441,7 +444,7 @@ protected:
 	TOptional<FTimecode> Timecode;
 
 	/** Which engine frame number this sample corresponds to. */
-	uint32 FrameNumber = 0;
+	std::atomic<uint64> FrameNumber = 0;
 
 	/** Image dimensions */
 	uint32 Stride = 0;
@@ -481,7 +484,7 @@ private:
 	TSharedPtr<FMediaIOCoreTextureSampleBase> OriginalSample;
 
 	/** Whether this sample's texture data is awaiting to be transferred by GPUDirect */
-	bool bIsAwaitingForGPUTransfer = false;
+	std::atomic<bool> bIsAwaitingForGPUTransfer = false;
 
 	/** Time offset evaluated on game thread for JITR */
 	double EvaluationOffsetInSeconds = 0;
diff --git a/Engine/Source/Runtime/MediaAssets/Private/Assets/TimeSynchronizableMediaSource.cpp b/Engine/Source/Runtime/MediaAssets/Private/Assets/TimeSynchronizableMediaSource.cpp
--- a/Engine/Source/Runtime/MediaAssets/Private/Assets/TimeSynchronizableMediaSource.cpp
+++ b/Engine/Source/Runtime/MediaAssets/Private/Assets/TimeSynchronizableMediaSource.cpp
@@ -36,7 +36,7 @@ int64 UTimeSynchronizableMediaSource::GetMediaOption(const FName& Key, int64 Def
 {
 	if (Key == TimeSynchronizableMedia::FrameDelay)
 	{
-		return FrameDelay;
+		return FMath::Clamp(FrameDelay, 0, 4);
 	}
 
 	return Super::GetMediaOption(Key, DefaultValue);
diff --git a/Engine/Source/Runtime/MediaAssets/Private/Misc/MediaTextureResource.cpp b/Engine/Source/Runtime/MediaAssets/Private/Misc/MediaTextureResource.cpp
--- a/Engine/Source/Runtime/MediaAssets/Private/Misc/MediaTextureResource.cpp
+++ b/Engine/Source/Runtime/MediaAssets/Private/Misc/MediaTextureResource.cpp
@@ -512,6 +512,11 @@ void FMediaTextureResource::Render(const FRenderParams& Params)
 						else
 						{
 							UpdateTextureReference(RHICmdList, OutTexture);
+							if (Converter->GetConverterInfoFlags() & IMediaTextureSampleConverter::ConverterInfoFlags_PreprocessOnly)
+							{
+								// ...followed by the built in conversion code as needed...
+								ConvertOrCopyNeeded = true;
+							}
 						}
 					}
 				}
@@ -861,6 +866,7 @@ bool FMediaTextureResource::RequiresConversion(const FTextureRHIRef& SampleTextu
 		(Format != PF_R8G8B8A8) &&
 		(Format != PF_FloatRGB) &&
 		(Format != PF_A16B16G16R16) &&
+		(Format != PF_A2B10G10R10) &&
 		(Format != PF_FloatRGBA));
 }
  
diff --git a/Engine/Source/Runtime/MediaAssets/Public/TimeSynchronizableMediaSource.h b/Engine/Source/Runtime/MediaAssets/Public/TimeSynchronizableMediaSource.h
--- a/Engine/Source/Runtime/MediaAssets/Public/TimeSynchronizableMediaSource.h
+++ b/Engine/Source/Runtime/MediaAssets/Public/TimeSynchronizableMediaSource.h
@@ -45,8 +45,11 @@ public:
 	UPROPERTY(BlueprintReadWrite, EditAnywhere, Category=Synchronization, meta=(DisplayName="Time Synchronization"))
 	bool bUseTimeSynchronization;
 
-	/** When using Time Synchronization, how many frame back should it read. */
-	UPROPERTY(EditAnywhere, Category=Synchronization, meta=(EditCondition="bUseTimeSynchronization"))
+	/** 
+	* When using Time Synchronization, how many frame back should it read. Capped at 4 frames due to the potential issues with buffer sizes.
+	* Framelocked modes don't support this feature.
+	*/
+	UPROPERTY(EditAnywhere, Category=Synchronization, meta=(EditCondition="bUseTimeSynchronization", ClampMin = "0", ClampMax = "4"))
 	int32 FrameDelay;
 
 	/** When not using Time Synchronization, how far back it time should it read. */
