Fix media IO using JITR even if the option is disabled.

MeidaIOCorePlayerBase would initialize jitr samples even when JITR was disabled, which caused the MediaIOCoreTextureSamples to use a sample converter that wouldn't work since we would never receive JITR samples.

diff --git a/Engine/Plugins/Media/MediaIOFramework/Source/MediaIOCore/Private/Player/MediaIOCorePlayerBase.cpp b/Engine/Plugins/Media/MediaIOFramework/Source/MediaIOCore/Private/Player/MediaIOCorePlayerBase.cpp
--- a/Engine/Plugins/Media/MediaIOFramework/Source/MediaIOCore/Private/Player/MediaIOCorePlayerBase.cpp
+++ b/Engine/Plugins/Media/MediaIOFramework/Source/MediaIOCore/Private/Player/MediaIOCorePlayerBase.cpp
@@ -302,18 +302,13 @@ void FMediaIOCorePlayerBase::TickTimeManagement()
 void FMediaIOCorePlayerBase::TickFetch(FTimespan DeltaTime, FTimespan Timecode)
 {
 	// Running JITR?
-	if (CurrentState == EMediaState::Playing)
+	if (CurrentState == EMediaState::Playing && IsJustInTimeRenderingEnabled())
 	{
-		// If not in Just In Time Mode and no samples received yet, don't provide the ProxySample otherwise rendering will be scheduled.
-		// Otherwise just in time implies that the sample most likely is going to arrive closer to the render thread / RHI thread command execution.
+		// Running in JITR mode implies that the sample most likely is going to arrive closer to the render thread / RHI thread command execution.
 		// AcquireJITRProxySampleInitialized is called on game thread before any rendering commands are executed for this frame.
 		// With just in time enabled, the Sample list could be empty and the sample could arrive much later. 
 		// The inheriting players must be responsible for handling the provision of JIT samples at the later stage.
-		if (Samples->NumVideoSamples() > 0 || IsJustInTimeRenderingEnabled())
-		{
-			// Create new JITR proxy sample
-			JITRSamples->ProxySample = AcquireJITRProxySampleInitialized();
-		}
+		JITRSamples->ProxySample = AcquireJITRProxySampleInitialized();
 	}
 
 	Samples->CacheSamplesState(GetTime());
