Media IO: JIT Timecode based sample picking was not taking FrameDelay into account correctly and latest sample was always picked.

diff --git a/Engine/Plugins/Media/MediaIOFramework/Source/MediaIOCore/Private/Player/MediaIOCorePlayerBase.cpp b/Engine/Plugins/Media/MediaIOFramework/Source/MediaIOCore/Private/Player/MediaIOCorePlayerBase.cpp
--- a/Engine/Plugins/Media/MediaIOFramework/Source/MediaIOCore/Private/Player/MediaIOCorePlayerBase.cpp
+++ b/Engine/Plugins/Media/MediaIOFramework/Source/MediaIOCore/Private/Player/MediaIOCorePlayerBase.cpp
@@ -1292,7 +1292,7 @@ TSharedPtr<FMediaIOCoreTextureSampleBase> FMediaIOCorePlayerBase::PickSampleToRe
 
 	// Apply time correction to the target time
 	const FTimespan RequestedOffsetTimespan = FTimespan::FromSeconds(InFrameInformation.EvaluationOffset);
-	const FTimespan TargetTimespanCorrected = TargetSampleTimespan + RequestedOffsetTimespan;
+	const FTimespan TargetTimespanCorrected = TargetSampleTimespan - RequestedOffsetTimespan;
 
 	// Go over the sample pool and find a sample closest to the target time
 	int32 ClosestIndex = -1;
