Media IO Player: Just in time sample picker was not providing the right frame time and Media Texture Resource was discarding samples.

diff --git a/Engine/Plugins/Media/MediaIOFramework/Source/MediaIOCore/Private/Player/MediaIOCorePlayerBase.cpp b/Engine/Plugins/Media/MediaIOFramework/Source/MediaIOCore/Private/Player/MediaIOCorePlayerBase.cpp
--- a/Engine/Plugins/Media/MediaIOFramework/Source/MediaIOCore/Private/Player/MediaIOCorePlayerBase.cpp
+++ b/Engine/Plugins/Media/MediaIOFramework/Source/MediaIOCore/Private/Player/MediaIOCorePlayerBase.cpp
@@ -1201,16 +1201,7 @@ TSharedPtr<FMediaIOCoreTextureSampleBase> FMediaIOCorePlayerBase::AcquireJITRPro
 	Args.Player    = AsShared().ToSharedPtr();
 	Args.Timecode = FApp::GetTimecode();
 
-	TOptional<FQualifiedFrameTime> CurrentFrameTime = FApp::GetCurrentFrameTime();
-	if (CurrentFrameTime.IsSet() && EvaluationType != EMediaIOSampleEvaluationType::Latest)
-	{
-		Args.Time = FTimespan::FromSeconds(CurrentFrameTime.GetValue().AsSeconds());
-	}
-	else
-	{
-		Args.Time = FTimespan::FromSeconds(GetPlatformSeconds());
-	}
-
+	Args.Time = CurrentTime;
 	Args.EvaluationOffsetInSeconds = GetEvaluationOffsetInSeconds();
 	Args.Converter = CreateTextureSampleConverter();
 	Args.FrameRate = VideoFrameRate;
