[MediaIOFramework, PixelCapture, PS2] Add: API surface to allow the resize to be executed in a custom capture pass

diff --git a/Engine/Plugins/Media/MediaIOFramework/Source/MediaIOCore/Private/MediaCaptureHelper.cpp b/Engine/Plugins/Media/MediaIOFramework/Source/MediaIOCore/Private/MediaCaptureHelper.cpp
--- a/Engine/Plugins/Media/MediaIOFramework/Source/MediaIOCore/Private/MediaCaptureHelper.cpp
+++ b/Engine/Plugins/Media/MediaIOFramework/Source/MediaIOCore/Private/MediaCaptureHelper.cpp
@@ -39,7 +39,17 @@ bool FMediaCaptureHelper::AreInputsValid(const UE::MediaCaptureData::FCaptureFra
 		return false;
 	}
 
-	if (Args.MediaCapture->DesiredCaptureOptions.ResizeMethod != EMediaCaptureResizeMethod::ResizeInRenderPass)
+	if (Args.MediaCapture->DesiredCaptureOptions.ResizeMethod == EMediaCaptureResizeMethod::ResizeInCapturePass &&
+		Args.MediaCapture->GetConversionOperation() != EMediaCaptureConversionOperation::CUSTOM)
+	{
+		UE_LOG(LogMediaIOCore, Error, TEXT("The capture will stop for '%s'. The resize method of ResizeInCapturePass expects GetConversionOperation() to return EMediaCaptureConversionOperation::CUSTOM")
+			, *Args.MediaCapture->MediaOutputName);
+
+		return false;	
+	}
+
+	if (Args.MediaCapture->DesiredCaptureOptions.ResizeMethod != EMediaCaptureResizeMethod::ResizeInRenderPass && 
+		Args.MediaCapture->DesiredCaptureOptions.ResizeMethod != EMediaCaptureResizeMethod::ResizeInCapturePass)
 	{
 		bool bFoundSizeMismatch = false;
 		FIntPoint RequestSize = FIntPoint::ZeroValue;
diff --git a/Engine/Plugins/Media/MediaIOFramework/Source/MediaIOCore/Public/MediaCapture.h b/Engine/Plugins/Media/MediaIOFramework/Source/MediaIOCore/Public/MediaCapture.h
--- a/Engine/Plugins/Media/MediaIOFramework/Source/MediaIOCore/Public/MediaCapture.h
+++ b/Engine/Plugins/Media/MediaIOFramework/Source/MediaIOCore/Public/MediaCapture.h
@@ -153,7 +153,9 @@ enum class EMediaCaptureResizeMethod : uint8
 	/** Resize the source that is being captured. This will resize the viewport when doing a viewport capture and will resize the render target when capturing a render target. Not valid for immediate capture of RHI resource */
 	ResizeSource,
 	/** Leaves the source as is, but will add a render pass where we resize the captured texture, leaving the source intact. This is useful when the output size is smaller than the captured source. */
-	ResizeInRenderPass
+	ResizeInRenderPass,
+	/**  Leaves the source as is, but expects the capture pass to fully re-render the source texture to the output texture. */
+	ResizeInCapturePass
 };
 
 
diff --git a/Engine/Plugins/Media/PixelCapture/Source/PixelCapture/Private/PixelCaptureCapturerMediaCapture.cpp b/Engine/Plugins/Media/PixelCapture/Source/PixelCapture/Private/PixelCaptureCapturerMediaCapture.cpp
--- a/Engine/Plugins/Media/PixelCapture/Source/PixelCapture/Private/PixelCaptureCapturerMediaCapture.cpp
+++ b/Engine/Plugins/Media/PixelCapture/Source/PixelCapture/Private/PixelCaptureCapturerMediaCapture.cpp
@@ -226,7 +226,7 @@ void FPixelCaptureCapturerMediaCapture::InitializeMediaCapture()
 	FMediaCaptureOptions CaptureOptions;
 	CaptureOptions.bSkipFrameWhenRunningExpensiveTasks = false;
 	CaptureOptions.OverrunAction = EMediaCaptureOverrunAction::Skip;
-	CaptureOptions.ResizeMethod = EMediaCaptureResizeMethod::None;
+	CaptureOptions.ResizeMethod = EMediaCaptureResizeMethod::ResizeInCapturePass;
 
 	FRHICaptureResourceDescription ResourceDescription;
 	ResourceDescription.PixelFormat = EPixelFormat::PF_B8G8R8A8;
diff --git a/Engine/Plugins/Media/PixelStreaming2/Source/PixelStreaming2/Private/VideoCapturer.cpp b/Engine/Plugins/Media/PixelStreaming2/Source/PixelStreaming2/Private/VideoCapturer.cpp
--- a/Engine/Plugins/Media/PixelStreaming2/Source/PixelStreaming2/Private/VideoCapturer.cpp
+++ b/Engine/Plugins/Media/PixelStreaming2/Source/PixelStreaming2/Private/VideoCapturer.cpp
@@ -153,13 +153,20 @@ namespace UE::PixelStreaming2
 					}
 					else
 					{
-						if (UPixelStreaming2PluginSettings::CVarCaptureUseFence.GetValueOnAnyThread())
+						if (UPixelStreaming2PluginSettings::CVarUseMediaCapture.GetValueOnAnyThread())
 						{
-							return FPixelCaptureCapturerRHI::Create(FinalScale);
+							return FPixelCaptureCapturerMediaCapture::Create(FinalScale, FinalFormat);
 						}
 						else
 						{
-							return FPixelCaptureCapturerRHIRDG::Create(FinalScale);
+							if (UPixelStreaming2PluginSettings::CVarCaptureUseFence.GetValueOnAnyThread())
+							{
+								return FPixelCaptureCapturerRHI::Create(FinalScale);
+							}
+							else
+							{
+								return FPixelCaptureCapturerRHIRDG::Create(FinalScale);
+							}
 						}
 					}
 				case PixelCaptureBufferFormat::FORMAT_I420:
