Load the GPUDirect library on demand rather than on startup module, unless the MediaIO.EnableGPUDirect cvar is enabled.

GPUDirect initialization is now done on the render thread.

Deprecate Aja.EnableGPUDirect and Blackmagic.EnableGPUDirect in favor of MediaIO.EnableGPUDirect

diff --git a/Engine/Plugins/Media/AjaMedia/Source/Aja/Private/OutputChannel.cpp b/Engine/Plugins/Media/AjaMedia/Source/Aja/Private/OutputChannel.cpp
--- a/Engine/Plugins/Media/AjaMedia/Source/Aja/Private/OutputChannel.cpp
+++ b/Engine/Plugins/Media/AjaMedia/Source/Aja/Private/OutputChannel.cpp
@@ -195,7 +195,10 @@ namespace AJA
 			, bInterlacedTest_ExpectFirstLineToBeWhite(false)
 			, InterlacedTest_FrameCounter(0)
 		{
-			TextureTransfer = FGPUTextureTransferModule::Get().GetTextureTransfer();
+			if (FGPUTextureTransferModule::Get().IsEnabled() && FGPUTextureTransferModule::Get().IsInitialized())
+			{
+				TextureTransfer = FGPUTextureTransferModule::Get().GetTextureTransfer();
+			}
 
 			if (InOptions.OutputNumberOfBuffers > 0)
 			{
diff --git a/Engine/Plugins/Media/AjaMedia/Source/AjaMediaOutput/Private/AjaMediaCapture.cpp b/Engine/Plugins/Media/AjaMedia/Source/AjaMediaOutput/Private/AjaMediaCapture.cpp
--- a/Engine/Plugins/Media/AjaMedia/Source/AjaMediaOutput/Private/AjaMediaCapture.cpp
+++ b/Engine/Plugins/Media/AjaMedia/Source/AjaMediaOutput/Private/AjaMediaCapture.cpp
@@ -26,10 +26,7 @@
 #include "Interfaces/IMainFrameModule.h"
 #endif
 
-static TAutoConsoleVariable<int32> CVarAjaEnableGPUDirect(
-	TEXT("Aja.EnableGPUDirect"), 0,
-	TEXT("Whether to enable GPU direct for faster video frame copies. (Experimental)"),
-	ECVF_RenderThreadSafe);
+static FAutoConsoleVariableDeprecated CVarAjaEnableGPUDirect_Deprecated(TEXT("Aja.EnableGPUDirect"), TEXT("MediaIO.EnableGPUDirect"), TEXT("5.6"));
 
 /* namespace AjaMediaCaptureDevice
 *****************************************************************************/
@@ -235,20 +232,25 @@ UAjaMediaCapture::UAjaMediaCapture()
 	, FrameRate(30, 1)
 	, WakeUpEvent(nullptr)
 {
-	bGPUTextureTransferAvailable = FGPUTextureTransferModule::Get().IsAvailable();
-	TextureTransfer = FGPUTextureTransferModule::Get().GetTextureTransfer();
+	if (!HasAnyFlags(RF_ClassDefaultObject | RF_ArchetypeObject))
+	{
+		if (FGPUTextureTransferModule::Get().IsInitialized())
+		{
+			TextureTransfer = FGPUTextureTransferModule::Get().GetTextureTransfer();
+		}
 
 #if WITH_EDITOR
-	if (GEditor)
-	{
-		// In editor, an asset re-save dialog can prevent AJA from cleaning up in the regular PreExit callback,
-		// So we have to do our cleanup before the regular callback is called.
-		IMainFrameModule& MainFrame = FModuleManager::LoadModuleChecked<IMainFrameModule>("MainFrame");
-		CanCloseEditorDelegateHandle = MainFrame.RegisterCanCloseEditor(IMainFrameModule::FMainFrameCanCloseEditor::CreateUObject(this, &UAjaMediaCapture::CleanupPreEditorExit));
-	}
+		if (GEditor)
+		{
+			// In editor, an asset re-save dialog can prevent AJA from cleaning up in the regular PreExit callback,
+			// So we have to do our cleanup before the regular callback is called.
+			IMainFrameModule& MainFrame = FModuleManager::LoadModuleChecked<IMainFrameModule>("MainFrame");
+			CanCloseEditorDelegateHandle = MainFrame.RegisterCanCloseEditor(IMainFrameModule::FMainFrameCanCloseEditor::CreateUObject(this, &UAjaMediaCapture::CleanupPreEditorExit));
+		}
 #else
-	FCoreDelegates::OnEnginePreExit.AddUObject(this, &UAjaMediaCapture::OnEnginePreExit);
+		FCoreDelegates::OnEnginePreExit.AddUObject(this, &UAjaMediaCapture::OnEnginePreExit);
 #endif
+	}
 }
 
 UAjaMediaCapture::~UAjaMediaCapture()
@@ -278,6 +280,12 @@ bool UAjaMediaCapture::ValidateMediaOutput() const
 bool UAjaMediaCapture::InitializeCapture()
 {
 	UAjaMediaOutput* AjaMediaSource = CastChecked<UAjaMediaOutput>(MediaOutput);
+
+	if (FGPUTextureTransferModule::Get().IsEnabled() && !FGPUTextureTransferModule::Get().IsInitialized())
+	{
+		// This will trigger an initialization of GPUDirect.
+		TextureTransfer = FGPUTextureTransferModule::Get().GetTextureTransfer();
+	}
 	
 	const bool bResult = InitAJA(AjaMediaSource);
 
@@ -375,7 +383,7 @@ void UAjaMediaCapture::StopCaptureImpl(bool bAllowPendingFrameToBeProcess)
 
 bool UAjaMediaCapture::ShouldCaptureRHIResource() const
 {
-	return bGPUTextureTransferAvailable && CVarAjaEnableGPUDirect.GetValueOnAnyThread() == 1;
+	return TextureTransfer && FGPUTextureTransferModule::Get().IsEnabled();
 }
 
 void UAjaMediaCapture::ApplyViewportTextureAlpha(TSharedPtr<FSceneViewport> InSceneViewport)
@@ -481,7 +489,7 @@ bool UAjaMediaCapture::InitAJA(UAjaMediaOutput* InAjaMediaOutput)
 		if (InAjaMediaOutput->OutputConfiguration.MediaConfiguration.MediaMode.Standard != EMediaIOStandardType::Progressive)
 		{
 			UE_LOG(LogAjaMediaOutput, Warning, TEXT("GPU DMA is not supported with interlaced, defaulting to regular path."));
-			bGPUTextureTransferAvailable = false;
+			TextureTransfer.Reset();
 		}
 	}
 
diff --git a/Engine/Plugins/Media/AjaMedia/Source/AjaMediaOutput/Public/AjaMediaCapture.h b/Engine/Plugins/Media/AjaMedia/Source/AjaMediaOutput/Public/AjaMediaCapture.h
--- a/Engine/Plugins/Media/AjaMedia/Source/AjaMediaOutput/Public/AjaMediaCapture.h
+++ b/Engine/Plugins/Media/AjaMedia/Source/AjaMediaOutput/Public/AjaMediaCapture.h
@@ -127,9 +127,6 @@ private:
 	/** Textures to release when the capture has stopped. Must be released after GPUTextureTransfer textures have been unregistered. */
 	TArray<FTextureRHIRef> TexturesToRelease;
 
-	/** Whether or not GPUTextureTransfer was initialized successfully. */
-	bool bGPUTextureTransferAvailable = false;
-
 	/** Handle for the delegate used to clean up AJA on editor shutdown. */
 	FDelegateHandle CanCloseEditorDelegateHandle;
 
diff --git a/Engine/Plugins/Media/BlackmagicMedia/Source/BlackmagicMediaOutput/Private/BlackmagicMediaCapture.cpp b/Engine/Plugins/Media/BlackmagicMedia/Source/BlackmagicMediaOutput/Private/BlackmagicMediaCapture.cpp
--- a/Engine/Plugins/Media/BlackmagicMedia/Source/BlackmagicMediaOutput/Private/BlackmagicMediaCapture.cpp
+++ b/Engine/Plugins/Media/BlackmagicMedia/Source/BlackmagicMediaOutput/Private/BlackmagicMediaCapture.cpp
@@ -24,10 +24,7 @@
 #include "EngineAnalytics.h"
 #endif
 
-TAutoConsoleVariable<int32> CVarBlackmagicEnableGPUDirect(
-	TEXT("Blackmagic.EnableGPUDirect"), 0,
-	TEXT("Whether to enable GPU direct for faster video frame copies. (Experimental)"),
-	ECVF_RenderThreadSafe);
+static FAutoConsoleVariableDeprecated CVarBlackmagicEnableGPUDirect_Deprecated(TEXT("Blackmagic.EnableGPUDirect"), TEXT("MediaIO.EnableGPUDirect"), TEXT("5.6"));
 
 bool bBlackmagicWritInputRawDataCmdEnable = false;
 static FAutoConsoleCommand BlackmagicWriteInputRawDataCmd(
@@ -391,10 +388,12 @@ UBlackmagicMediaCapture::UBlackmagicMediaCapture()
 	, WakeUpEvent(nullptr)
 	, LastFrameDropCount_BlackmagicThread(0)
 {
-	bGPUTextureTransferAvailable = FGPUTextureTransferModule::Get().IsAvailable();
-	if (bGPUTextureTransferAvailable)
+	if (!HasAnyFlags(RF_ClassDefaultObject | RF_ArchetypeObject))
 	{
-		TextureTransfer = FGPUTextureTransferModule::Get().GetTextureTransfer();
+		if (FGPUTextureTransferModule::Get().IsInitialized())
+		{
+			TextureTransfer = FGPUTextureTransferModule::Get().GetTextureTransfer();
+		}
 	}
 }
 
@@ -414,6 +413,13 @@ bool UBlackmagicMediaCapture::InitializeCapture()
 {
 	UBlackmagicMediaOutput* BlackmagicMediaOutput = CastChecked<UBlackmagicMediaOutput>(MediaOutput);
 	BlackmagicMediaOutputPixelFormat = BlackmagicMediaOutput->PixelFormat;
+
+	if (FGPUTextureTransferModule::Get().IsEnabled() && !FGPUTextureTransferModule::Get().IsInitialized())
+	{
+		// This will trigger an initialization of GPUDirect.
+		TextureTransfer = FGPUTextureTransferModule::Get().GetTextureTransfer();
+	}
+
 	bool bInitialized = InitBlackmagic(BlackmagicMediaOutput);
 	if(bInitialized)
 	{
@@ -488,8 +494,7 @@ void UBlackmagicMediaCapture::StopCaptureImpl(bool bAllowPendingFrameToBeProcess
 
 bool UBlackmagicMediaCapture::ShouldCaptureRHIResource() const
 {
-	// Todo: also test if dvp was initialized correctly.
-	return bGPUTextureTransferAvailable && CVarBlackmagicEnableGPUDirect.GetValueOnAnyThread() == 1;
+	return TextureTransfer && FGPUTextureTransferModule::Get().IsEnabled();
 }
 
 void UBlackmagicMediaCapture::ApplyViewportTextureAlpha(TSharedPtr<FSceneViewport> InSceneViewport)
@@ -652,7 +657,7 @@ bool UBlackmagicMediaCapture::InitBlackmagic(UBlackmagicMediaOutput* InBlackmagi
 		if (InBlackmagicMediaOutput->OutputConfiguration.MediaConfiguration.MediaMode.Standard != EMediaIOStandardType::Progressive)
 		{
 			UE_LOG(LogBlackmagicMediaOutput, Warning, TEXT("GPU DMA is not supported with interlaced, defaulting to regular path."));
-			bGPUTextureTransferAvailable = false;
+			TextureTransfer.Reset();
 		}
 	}
 
diff --git a/Engine/Plugins/Media/BlackmagicMedia/Source/BlackmagicMediaOutput/Public/BlackmagicMediaCapture.h b/Engine/Plugins/Media/BlackmagicMedia/Source/BlackmagicMediaOutput/Public/BlackmagicMediaCapture.h
--- a/Engine/Plugins/Media/BlackmagicMedia/Source/BlackmagicMediaOutput/Public/BlackmagicMediaCapture.h
+++ b/Engine/Plugins/Media/BlackmagicMedia/Source/BlackmagicMediaOutput/Public/BlackmagicMediaCapture.h
@@ -115,9 +115,6 @@ private:
 	/** Textures to release when the capture has stopped. Must be released after GPUTextureTransfer textures have been unregistered. */
 	TArray<FTextureRHIRef> TexturesToRelease;
 
-	/** Whether or not GPUTextureTransfer was initialized successfully. */
-	bool bGPUTextureTransferAvailable = false;
-
 	/** The last time OutputAudio was called. Used to adjust the number of audio samples to grab on each frame depending on the frame time. */
 	double OutputAudioTimestamp;
 
diff --git a/Engine/Plugins/Media/MediaIOFramework/MediaIOFramework.uplugin b/Engine/Plugins/Media/MediaIOFramework/MediaIOFramework.uplugin
--- a/Engine/Plugins/Media/MediaIOFramework/MediaIOFramework.uplugin
+++ b/Engine/Plugins/Media/MediaIOFramework/MediaIOFramework.uplugin
@@ -27,7 +27,7 @@
 		{
 			"Name": "GPUTextureTransfer",
 			"Type": "Runtime",
-			"LoadingPhase": "PostConfigInit"
+			"LoadingPhase": "Default"
 		}
 	],
 	"Plugins": [
diff --git a/Engine/Plugins/Media/MediaIOFramework/Source/GPUTextureTransfer/Private/GPUTextureTransferModule.cpp b/Engine/Plugins/Media/MediaIOFramework/Source/GPUTextureTransfer/Private/GPUTextureTransferModule.cpp
--- a/Engine/Plugins/Media/MediaIOFramework/Source/GPUTextureTransfer/Private/GPUTextureTransferModule.cpp
+++ b/Engine/Plugins/Media/MediaIOFramework/Source/GPUTextureTransfer/Private/GPUTextureTransferModule.cpp
@@ -5,6 +5,7 @@
 #if DVP_SUPPORTED_PLATFORM
 #include "D3D11TextureTransfer.h"
 #include "D3D12TextureTransfer.h"
+#include "RenderingThread.h"
 #include "VulkanTextureTransfer.h"
 #include "TextureTransferBase.h"
 #endif
@@ -31,6 +32,11 @@
 
 DEFINE_LOG_CATEGORY(LogGPUTextureTransfer);
 
+static TAutoConsoleVariable<bool> CVarMediaIOEnableGPUDirect(
+	TEXT("MediaIO.EnableGPUDirect"), false,
+	TEXT("Whether to enable GPU direct for faster video frame copies. (Experimental)"),
+	ECVF_RenderThreadSafe);
+
 namespace 
 {
 	auto ConvertRHI = [](ERHIInterfaceType RHI)
@@ -45,7 +51,37 @@ namespace
 	};
 }
 
+bool FGPUTextureTransferModule::IsAvailable()
+{
+#if DVP_SUPPORTED_PLATFORM
+	return FGPUTextureTransferModule::Get().IsInitialized() && FGPUTextureTransferModule::Get().IsEnabled();
+#else
+	return false;
+#endif
+}
+
 void FGPUTextureTransferModule::StartupModule()
+{
+	if (CVarMediaIOEnableGPUDirect.GetValueOnAnyThread())
+	{
+		FCoreDelegates::OnAllModuleLoadingPhasesComplete.AddRaw(this, &FGPUTextureTransferModule::Initialize);
+	}
+
+	CVarMediaIOEnableGPUDirect.AsVariable()->OnChangedDelegate().AddRaw(this, &FGPUTextureTransferModule::OnEnableGPUDirectCVarChange);
+
+	// Cache this information since GetGPUDriverInfo has to be called on game thread because of a call to GetValueOnGameThread.
+	CachedDriverInfo = FPlatformMisc::GetGPUDriverInfo(GRHIAdapterName);
+}
+
+void FGPUTextureTransferModule::ShutdownModule()
+{
+	if (bInitialized)
+	{
+		UninitializeTextureTransfer();
+	}
+}
+
+void FGPUTextureTransferModule::Initialize()
 {
 	if (FApp::CanEverRender())
 	{
@@ -66,20 +102,25 @@ void FGPUTextureTransferModule::StartupModule()
 
 			TransferObjects.AddDefaulted(RHI_MAX);
 
-			// Since this module is started before the RHI is initialized, we have to delay initialization to later. 
-			FCoreDelegates::OnAllModuleLoadingPhasesComplete.AddRaw(this, &FGPUTextureTransferModule::InitializeTextureTransfer);
-			// Same for shutdown, uninitialize ourselves before library is unloaded
-			FCoreDelegates::OnEnginePreExit.AddRaw(this, &FGPUTextureTransferModule::UninitializeTextureTransfer);
+			InitializeTextureTransfer();
 		}
 	}
 }
 
-void FGPUTextureTransferModule::ShutdownModule()
-{
-}
-
 UE::GPUTextureTransfer::TextureTransferPtr FGPUTextureTransferModule::GetTextureTransfer()
 {
+#if DVP_SUPPORTED_PLATFORM
+	if (!IsInitialized())
+	{
+		UE_LOG(LogGPUTextureTransfer, Warning, TEXT("GetTextureTransfer was called without initializing the library. This will cause a hitch since we have to block while waiting for the library to finish initializing."));
+		Initialize();
+
+		// Initialization is done on the rendering thread.
+		FlushRenderingCommands();
+	}
+#endif
+
+
 	if (FApp::CanEverRender())
 	{
 #if DVP_SUPPORTED_PLATFORM
@@ -100,15 +141,20 @@ UE::GPUTextureTransfer::TextureTransferPtr FGPUTextureTransferModule::GetTexture
 	return nullptr;
 }
 
-bool FGPUTextureTransferModule::IsAvailable()
+bool FGPUTextureTransferModule::IsInitialized() const
 {
 #if DVP_SUPPORTED_PLATFORM
-	return FModuleManager::Get().IsModuleLoaded("GPUTextureTransfer");
+	return bInitialized;
 #else
 	return false;
 #endif
 }
 
+bool FGPUTextureTransferModule::IsEnabled() const
+{
+	return CVarMediaIOEnableGPUDirect.GetValueOnAnyThread();
+}
+
 FGPUTextureTransferModule& FGPUTextureTransferModule::Get()
 {
 	return FModuleManager::LoadModuleChecked<FGPUTextureTransferModule>("GPUTextureTransfer");
@@ -142,7 +188,8 @@ bool FGPUTextureTransferModule::LoadGPUDirectBinary()
 void FGPUTextureTransferModule::InitializeTextureTransfer()
 {
 #if DVP_SUPPORTED_PLATFORM
-	bIsGPUTextureTransferAvailable = true;
+
+	bool bInitializationSuccess = true;
 
 	static const TArray<FString> SupportedGPUPrefixes = {
 		TEXT("RTX A4"),
@@ -151,31 +198,29 @@ void FGPUTextureTransferModule::InitializeTextureTransfer()
 		TEXT("Quadro")
 	};
 
-	// This must be called on game thread 
-	const FGPUDriverInfo GPUDriverInfo = FPlatformMisc::GetGPUDriverInfo(GRHIAdapterName);
-	bIsGPUTextureTransferAvailable = GPUDriverInfo.IsNVIDIA() && !GPUDriverInfo.DeviceDescription.Contains(TEXT("Tesla"));
+	bInitializationSuccess = CachedDriverInfo.IsNVIDIA() && !CachedDriverInfo.DeviceDescription.Contains(TEXT("Tesla"));
 
 	const bool bRenderDocAttached = FParse::Param(FCommandLine::Get(), TEXT("AttachRenderDoc"));
 	if (bRenderDocAttached)
 	{
-		bIsGPUTextureTransferAvailable = false;
+		bInitializationSuccess = false;
 		// Render doc clashes with GPU Direct.
 		UE_LOG(LogGPUTextureTransfer, Display, TEXT("GPU Texture Transfer disabled because render is attached."))
 	}
-	
-	if (bIsGPUTextureTransferAvailable)
+
+	if (bInitializationSuccess)
 	{
-		bIsGPUTextureTransferAvailable = false;
+		bInitializationSuccess = false;
 		for (const FString& GPUPrefix : SupportedGPUPrefixes)
 		{
-			if (GPUDriverInfo.DeviceDescription.Contains(GPUPrefix))
+			if (CachedDriverInfo.DeviceDescription.Contains(GPUPrefix))
 			{
-				bIsGPUTextureTransferAvailable = true;
+				bInitializationSuccess = true;
 				break;
 			}
 		}
 	}
-	if (!bIsGPUTextureTransferAvailable)
+	if (!bInitializationSuccess)
 	{
 		return;
 	}
@@ -185,45 +230,50 @@ void FGPUTextureTransferModule::InitializeTextureTransfer()
 		return;
 	}
 
-	UE::GPUTextureTransfer::TextureTransferPtr TextureTransfer;
+	ENQUEUE_RENDER_COMMAND(InitializeTextureTransfer)([this](FRHICommandListImmediate&)
+	{
+		UE::GPUTextureTransfer::TextureTransferPtr TextureTransfer;
 
-	UE::GPUTextureTransfer::ERHI RHI = ConvertRHI(RHIGetInterfaceType());
+		UE::GPUTextureTransfer::ERHI RHI = ConvertRHI(RHIGetInterfaceType());
 
-	switch (RHI)
-	{
-	case UE::GPUTextureTransfer::ERHI::D3D11:
-		TextureTransfer = MakeShared<UE::GPUTextureTransfer::Private::FD3D11TextureTransfer>();
-		break;
-	case UE::GPUTextureTransfer::ERHI::D3D12:
-		TextureTransfer = MakeShared<UE::GPUTextureTransfer::Private::FD3D12TextureTransfer>();
-		break;
-	case UE::GPUTextureTransfer::ERHI::Vulkan:
-		TextureTransfer = MakeShared<UE::GPUTextureTransfer::Private::FVulkanTextureTransfer>();
-		break;
-	default:
-		ensureAlways(false);
-		break;
-	}
+		switch (RHI)
+		{
+		case UE::GPUTextureTransfer::ERHI::D3D11:
+			TextureTransfer = MakeShared<UE::GPUTextureTransfer::Private::FD3D11TextureTransfer>();
+			break;
+		case UE::GPUTextureTransfer::ERHI::D3D12:
+			TextureTransfer = MakeShared<UE::GPUTextureTransfer::Private::FD3D12TextureTransfer>();
+			break;
+		case UE::GPUTextureTransfer::ERHI::Vulkan:
+			TextureTransfer = MakeShared<UE::GPUTextureTransfer::Private::FVulkanTextureTransfer>();
+			break;
+		default:
+			ensureAlways(false);
+			break;
+		}
 
-	UE::GPUTextureTransfer::FInitializeDMAArgs InitializeArgs;
-	InitializeArgs.RHI = RHI;
-	InitializeArgs.RHIDevice = GDynamicRHI->RHIGetNativeDevice();
-	InitializeArgs.RHICommandQueue = GDynamicRHI->RHIGetNativeGraphicsQueue();
+		UE::GPUTextureTransfer::FInitializeDMAArgs InitializeArgs;
+		InitializeArgs.RHI = RHI;
+		InitializeArgs.RHIDevice = GDynamicRHI->RHIGetNativeDevice();
+		InitializeArgs.RHICommandQueue = GDynamicRHI->RHIGetNativeGraphicsQueue();
 #if VULKAN_PLATFORM
-	if (RHI == UE::GPUTextureTransfer::ERHI::Vulkan)
-	{
-		IVulkanDynamicRHI* DynRHI = GetIVulkanDynamicRHI();
-		InitializeArgs.VulkanInstance = DynRHI->RHIGetVkInstance();
-		FMemory::Memcpy(InitializeArgs.RHIDeviceUUID, DynRHI->RHIGetVulkanDeviceUUID(), 16);
-	}
+		if (RHI == UE::GPUTextureTransfer::ERHI::Vulkan)
+		{
+			IVulkanDynamicRHI* DynRHI = GetIVulkanDynamicRHI();
+			InitializeArgs.VulkanInstance = DynRHI->RHIGetVkInstance();
+			FMemory::Memcpy(InitializeArgs.RHIDeviceUUID, DynRHI->RHIGetVulkanDeviceUUID(), 16);
+		}
 #endif
 
-	const uint8 RHIIndex = static_cast<uint8>(RHI);
-	UE_LOG(LogGPUTextureTransfer, Display, TEXT("Initializing GPU Texture transfer"));
-	if (TextureTransfer->Initialize(InitializeArgs))
-	{
-		TransferObjects[RHIIndex] = TextureTransfer;
-	}
+		const uint8 RHIIndex = static_cast<uint8>(RHI);
+		UE_LOG(LogGPUTextureTransfer, Display, TEXT("Initializing GPU Texture transfer"));
+		if (TextureTransfer->Initialize(InitializeArgs))
+		{
+			TransferObjects[RHIIndex] = TextureTransfer;
+		}
+
+		bInitialized = true;
+	});
 #endif // DVP_SUPPORTED_PLATFORM
 }
 
@@ -240,6 +290,14 @@ void FGPUTextureTransferModule::UninitializeTextureTransfer()
 #endif
 }
 
+void FGPUTextureTransferModule::OnEnableGPUDirectCVarChange(IConsoleVariable* ConsoleVariable)
+{
+	if (ConsoleVariable->GetBool() && !IsInitialized())
+	{
+		Initialize();
+	}
+}
+
 namespace UE::GPUTextureTransfer
 {
 	struct FTextureTransfersWrapper
diff --git a/Engine/Plugins/Media/MediaIOFramework/Source/GPUTextureTransfer/Public/GPUTextureTransferModule.h b/Engine/Plugins/Media/MediaIOFramework/Source/GPUTextureTransfer/Public/GPUTextureTransferModule.h
--- a/Engine/Plugins/Media/MediaIOFramework/Source/GPUTextureTransfer/Public/GPUTextureTransferModule.h
+++ b/Engine/Plugins/Media/MediaIOFramework/Source/GPUTextureTransfer/Public/GPUTextureTransferModule.h
@@ -3,6 +3,7 @@
 #pragma once
 
 #include "CoreMinimal.h"
+#include "GenericPlatform/GenericPlatformDriver.h"
 #include "GPUTextureTransfer.h"
 #include "Logging/LogMacros.h"
 #include "Modules/ModuleInterface.h"
@@ -19,23 +20,54 @@ namespace UE::GPUTextureTransfer
 class GPUTEXTURETRANSFER_API FGPUTextureTransferModule : public IModuleInterface
 {
 public:
+	static FGPUTextureTransferModule& Get();
+
+	UE_DEPRECATED(5.6, "Use IsInitialized() and IsEnabled() instead.")
+	static bool IsAvailable();
+
+	//~ Begin IModuleInterface
 	virtual void StartupModule() override;
 	virtual void ShutdownModule() override;
+	//~ End IModuleInterface
 
+	/** Load the DVP DLL and intialize the API.Note: This is a blocking call and could take upwards of 2s. */
+	void Initialize();
+
+	/** Get a texture transfer object that acts as a layer above the DVP library. */
 	UE::GPUTextureTransfer::TextureTransferPtr GetTextureTransfer();
 
-	static bool IsAvailable();
-	static FGPUTextureTransferModule& Get();
+	/** Returns whether the DVP library was loaded and initialized. If this returns false, GetTextureTransfer will trigger the library to initialize. */
+	bool IsInitialized() const;
+
+	/** Returns whether GPUDirect is enabled. This reflects the value of MediaIO.EnableGPUDirect. You should check IsAvailable to know if the library is initialized. */
+	bool IsEnabled() const;
 
 private:
+	/** Load the DVP dll. */
 	bool LoadGPUDirectBinary();
+
+	/** Initialize the DVP library and creates the ITextureTransfer objects. */
 	void InitializeTextureTransfer();
+
+	/** Clean up and uninitializes the DVP library. */
 	void UninitializeTextureTransfer();
 
+	/** Initializes the DVP library when MediaIO.EnableGPUDirect is set to 1. */
+	void OnEnableGPUDirectCVarChange(class IConsoleVariable* ConsoleVariable);
+	
 private:
+	/** Max number of RHIs that GPUTextureTransfer supports. */
 	static constexpr uint8 RHI_MAX = static_cast<uint8>(UE::GPUTextureTransfer::ERHI::RHI_MAX);
+
+	/** DVP DLL Handle. */
 	void* TextureTransferHandle = nullptr;
-	std::atomic<bool> bIsGPUTextureTransferAvailable = false;
+
+	/** Tracks whether the dvp library and texture transfer objects were successfully created. */
+	std::atomic<bool> bInitialized = false;
 	
+	/** Texture transfer objects. */
 	TArray<UE::GPUTextureTransfer::TextureTransferPtr> TransferObjects;
+
+	/** Cached information about the GPU Driver. */
+	FGPUDriverInfo CachedDriverInfo;
 };
