Modify TextureTransferBase to reduce locking.
Remove DVPBegin / DVPEnd from the TransferTexture method to move them into their own functions.
This should enable better batching of commands on the output thread.

diff --git a/Engine/Plugins/Media/MediaIOFramework/Source/GPUTextureTransfer/Private/TextureTransferBase.cpp b/Engine/Plugins/Media/MediaIOFramework/Source/GPUTextureTransfer/Private/TextureTransferBase.cpp
--- a/Engine/Plugins/Media/MediaIOFramework/Source/GPUTextureTransfer/Private/TextureTransferBase.cpp
+++ b/Engine/Plugins/Media/MediaIOFramework/Source/GPUTextureTransfer/Private/TextureTransferBase.cpp
@@ -99,23 +99,55 @@ namespace UE::GPUTextureTransfer::Private
 		return true;
 	}
 
+	void FTextureTransferBase::WaitForGPU(FRHITexture* InRHITexture)
+	{
+		FTextureInfo* TextureInfo = nullptr;
+		{
+			FScopeLock ScopeLock{ &CriticalSection };
+			TextureInfo = RegisteredTextures.Find(InRHITexture);	
+		}
+
+		TRACE_CPUPROFILER_EVENT_SCOPE_TEXT(*FString::Printf(TEXT("GPUDirect::WaitForGPU (%d)"), TextureInfo->DVPHandle));
+		DVP_CALL(dvpBegin());
+		DVP_CALL(dvpMapBufferWaitDVP(TextureInfo->DVPHandle));
+		DVP_CALL(dvpEnd());
+	}
+
+	void FTextureTransferBase::ThreadPrep()
+	{
+		TRACE_CPUPROFILER_EVENT_SCOPE(dvpBegin);
+		DVP_CALL(dvpBegin());
+	}
+
+	void FTextureTransferBase::ThreadCleanup()
+	{
+		TRACE_CPUPROFILER_EVENT_SCOPE(dvpEnd);
+		DVP_CALL(dvpEnd());
+	}
+
 	bool FTextureTransferBase::BeginSync(void* InBuffer, ETransferDirection TransferDirection)
 	{
+		TRACE_CPUPROFILER_EVENT_SCOPE(BeginSync);
 		bool bSuccess = true;
 		LOG_PERF(BeginSync);
 		
-		FScopeLock ScopeLock{ &CriticalSection };
+		FExternalBufferInfo* Info = nullptr;
+		{
+			FScopeLock ScopeLock{ &CriticalSection };
+			Info = RegisteredBuffers.Find(InBuffer);
+		}
 
-		if (FExternalBufferInfo* Info = RegisteredBuffers.Find(InBuffer))
+		if (ensureAlways(Info))
 		{
 			if (Info->GPUMemorySync && Info->SystemMemorySync)
 			{
+				TRACE_CPUPROFILER_EVENT_SCOPE_TEXT(*FString::Printf(TEXT("BeginSync Texture: %d, Buffer %d"), Info->PendingTextureHandle, Info->DVPHandle));
 				Info->GPUMemorySync->AcquireValue++;
+
 				Info->SystemMemorySync->ReleaseValue++;
 
 				if (TransferDirection == ETransferDirection::GPU_TO_CPU)
 				{
-					DVP_CALL(dvpBegin());
 					constexpr uint64 NanosecondsToWait = 500000000; // 0.5 seconds
 					DVPStatus SyncStatus = dvpSyncObjClientWaitPartial(Info->GPUMemorySync->DVPSyncObject, Info->GPUMemorySync->AcquireValue, NanosecondsToWait);
 					if (SyncStatus != DVP_STATUS_OK)
@@ -123,8 +155,6 @@ namespace UE::GPUTextureTransfer::Private
 						bSuccess = false;
 						UE_LOG(LogGPUTextureTransfer, Error, TEXT("GPU Direct failed to sync."));
 					}
-
-					DVP_CALL(dvpEnd());
 				}
 			}
 			else
@@ -140,6 +170,8 @@ namespace UE::GPUTextureTransfer::Private
 
 	void FTextureTransferBase::EndSync(void* InBuffer)
 	{
+		TRACE_CPUPROFILER_EVENT_SCOPE(EndSync);
+
 		FScopeLock ScopeLock{ &CriticalSection };
 		if (FExternalBufferInfo* Info = RegisteredBuffers.Find(InBuffer))
 		{
@@ -147,11 +179,15 @@ namespace UE::GPUTextureTransfer::Private
 			{
 				Info->SystemMemorySync->SetValue(Info->SystemMemorySync->ReleaseValue);
 			}
+
+			Info->PendingTextureHandle = 0;
 		}
 	}
 
 	bool FTextureTransferBase::TransferTexture(void* InBuffer, FRHITexture* InRHITexture, ETransferDirection TransferDirection)
 	{
+		TRACE_CPUPROFILER_EVENT_SCOPE(TransferTexture);
+
 		FScopeLock ScopeLock{ &CriticalSection };
 
 		FExternalBufferInfo* BufferInfo = RegisteredBuffers.Find(InBuffer);
@@ -171,29 +207,36 @@ namespace UE::GPUTextureTransfer::Private
 
 		if (BufferInfo->GPUMemorySync && BufferInfo->SystemMemorySync)
 		{
+			//UE_LOG(LogGPUTextureTransfer, Verbose, TEXT("Transferring texture Buffer: %d")
+
 			BufferInfo->GPUMemorySync->ReleaseValue++;
+			BufferInfo->PendingTextureHandle = TextureInfo->DVPHandle;
 
-			DVP_CALL(dvpBegin());
-			DVP_CALL(dvpMapBufferWaitDVP(TextureInfo->DVPHandle));
-			
 			DVPStatus Status = DVP_STATUS_OK;
-			if (TransferDirection == ETransferDirection::GPU_TO_CPU)
-			{
-				Status = dvpMemcpy2D(TextureInfo->DVPHandle, BufferInfo->SystemMemorySync->DVPSyncObject,
-					BufferInfo->SystemMemorySync->AcquireValue, DVP_TIMEOUT_IGNORED, BufferInfo->DVPHandle,
-					BufferInfo->GPUMemorySync->DVPSyncObject, BufferInfo->GPUMemorySync->ReleaseValue, 0, 0, BufferInfo->Height, BufferInfo->Width);
-			}
-			else
+
 			{
-				Status = dvpMemcpy2D(BufferInfo->DVPHandle, BufferInfo->SystemMemorySync->DVPSyncObject,
-					BufferInfo->SystemMemorySync->AcquireValue, DVP_TIMEOUT_IGNORED, TextureInfo->DVPHandle,
-					BufferInfo->GPUMemorySync->DVPSyncObject, BufferInfo->GPUMemorySync->ReleaseValue, 0, 0, BufferInfo->Height, BufferInfo->Width);
+				TRACE_CPUPROFILER_EVENT_SCOPE_TEXT(*FString::Printf(TEXT("dvpMemcpy2D (%d)"), TextureInfo->DVPHandle));
+
+				if (TransferDirection == ETransferDirection::GPU_TO_CPU)
+				{
+					Status = dvpMemcpy2D(TextureInfo->DVPHandle, BufferInfo->SystemMemorySync->DVPSyncObject,
+						BufferInfo->SystemMemorySync->AcquireValue, DVP_TIMEOUT_IGNORED, BufferInfo->DVPHandle,
+						BufferInfo->GPUMemorySync->DVPSyncObject, BufferInfo->GPUMemorySync->ReleaseValue, 0, 0, BufferInfo->Height, BufferInfo->Width);
+				}
+				else
+				{
+					Status = dvpMemcpy2D(BufferInfo->DVPHandle, BufferInfo->SystemMemorySync->DVPSyncObject,
+						BufferInfo->SystemMemorySync->AcquireValue, DVP_TIMEOUT_IGNORED, TextureInfo->DVPHandle,
+						BufferInfo->GPUMemorySync->DVPSyncObject, BufferInfo->GPUMemorySync->ReleaseValue, 0, 0, BufferInfo->Height, BufferInfo->Width);
+				}
 			}
 
 			BufferInfo->SystemMemorySync->AcquireValue++;
 
-			DVP_CALL(dvpMapBufferEndDVP(TextureInfo->DVPHandle));
-			DVP_CALL(dvpEnd())
+			{
+				TRACE_CPUPROFILER_EVENT_SCOPE_TEXT(*FString::Printf(TEXT("dvpMapBufferEndDVP (%d)"), TextureInfo->DVPHandle));
+				DVP_CALL(dvpMapBufferEndDVP(TextureInfo->DVPHandle));
+			}
 
 			if (Status != DVP_STATUS_OK)
 			{
@@ -264,9 +307,11 @@ namespace UE::GPUTextureTransfer::Private
 	void FTextureTransferBase::RegisterTexture(const FRegisterDMATextureArgs& Args)
 	{
 		FScopeLock ScopeLock{ &CriticalSection };
+
 		if (!RegisteredTextures.Contains(Args.RHITexture))
 		{
 			FTextureInfo Info;
+			TRACE_CPUPROFILER_EVENT_SCOPE(CreateGPUResource_Impl);
 			DVP_CALL(CreateGPUResource_Impl(Args, &Info));
 			RegisteredTextures.Add({ Args.RHITexture, std::move(Info) });
 		}
@@ -321,19 +366,25 @@ namespace UE::GPUTextureTransfer::Private
 
 	void FTextureTransferBase::LockTexture(FRHITexture* InRHITexture)
 	{
+		TRACE_CPUPROFILER_EVENT_SCOPE(TEXT("LockTexture"));
 		FScopeLock ScopeLock{ &CriticalSection };
+
 		if (FTextureInfo* TextureInfo = RegisteredTextures.Find(InRHITexture))
 		{
+			TRACE_CPUPROFILER_EVENT_SCOPE_TEXT(*FString::Printf(TEXT("MapBufferWaitAPI_Impl (%d)"), TextureInfo->DVPHandle));
 			DVP_CALL(MapBufferWaitAPI_Impl(TextureInfo->DVPHandle));
 		}
 	}
 
 	void FTextureTransferBase::UnlockTexture(FRHITexture* InRHITexture)
 	{
+		TRACE_CPUPROFILER_EVENT_SCOPE(TEXT("UnlockTexture"));
 		LOG_PERF(UnlockTexture);
+
 		FScopeLock ScopeLock{ &CriticalSection };
 		if (FTextureInfo* TextureInfo = RegisteredTextures.Find(InRHITexture))
 		{
+			TRACE_CPUPROFILER_EVENT_SCOPE_TEXT(*FString::Printf(TEXT("MapBufferEndAPI_Impl (%d)"), TextureInfo->DVPHandle));
 			DVP_CALL(MapBufferEndAPI_Impl(TextureInfo->DVPHandle));
 		}
 	}
@@ -358,11 +409,13 @@ namespace UE::GPUTextureTransfer::Private
 
 	DVPStatus FTextureTransferBase::MapBufferWaitAPI_Impl(DVPBufferHandle Handle) const
 	{
+		TRACE_CPUPROFILER_EVENT_SCOPE_TEXT(*FString::Printf(TEXT("dvpMapBufferWaitAPI (%d)"), Handle));
 		return dvpMapBufferWaitAPI(Handle);
 	}
 
 	DVPStatus FTextureTransferBase::MapBufferEndAPI_Impl(DVPBufferHandle Handle) const
 	{
+		TRACE_CPUPROFILER_EVENT_SCOPE_TEXT(*FString::Printf(TEXT("dvpMapBufferEndAPI (%d)"), Handle));
 		return dvpMapBufferEndAPI(Handle);
 	}
 
diff --git a/Engine/Plugins/Media/MediaIOFramework/Source/GPUTextureTransfer/Private/TextureTransferBase.h b/Engine/Plugins/Media/MediaIOFramework/Source/GPUTextureTransfer/Private/TextureTransferBase.h
--- a/Engine/Plugins/Media/MediaIOFramework/Source/GPUTextureTransfer/Private/TextureTransferBase.h
+++ b/Engine/Plugins/Media/MediaIOFramework/Source/GPUTextureTransfer/Private/TextureTransferBase.h
@@ -19,6 +19,9 @@ namespace UE::GPUTextureTransfer::Private
 		virtual ~FTextureTransferBase() {};
 
 		//~ Begin ITextureTransfer interface
+		virtual void WaitForGPU(FRHITexture* InRHITexture) override;
+		virtual void ThreadPrep() override;
+		virtual void ThreadCleanup() override;
 		virtual bool BeginSync(void* InBuffer, ETransferDirection TransferDirection) override;
 		virtual void EndSync(void* InBuffer) override;
 		virtual bool TransferTexture(void* InBuffer, FRHITexture* InRHITexture, ETransferDirection TransferDirection) override;
@@ -58,6 +61,9 @@ namespace UE::GPUTextureTransfer::Private
 			TUniquePtr<DVPSync> SystemMemorySync;
 			TUniquePtr<DVPSync> GPUMemorySync;
 			DVPBufferHandle DVPHandle = 0;
+
+			// Handle to the texture currently being transferred to this buffer.
+			DVPBufferHandle PendingTextureHandle = 0;
 		};
 
 		struct FTextureInfo
diff --git a/Engine/Plugins/Media/MediaIOFramework/Source/GPUTextureTransfer/Public/GPUTextureTransfer.h b/Engine/Plugins/Media/MediaIOFramework/Source/GPUTextureTransfer/Public/GPUTextureTransfer.h
--- a/Engine/Plugins/Media/MediaIOFramework/Source/GPUTextureTransfer/Public/GPUTextureTransfer.h
+++ b/Engine/Plugins/Media/MediaIOFramework/Source/GPUTextureTransfer/Public/GPUTextureTransfer.h
@@ -74,6 +74,13 @@ namespace UE::GPUTextureTransfer
 	public:
 		virtual ~ITextureTransfer() {};
 
+		/** Wait for a signal that the texture is ready to use by DVP. (Experimental) */
+		virtual void WaitForGPU(FRHITexture* InRHITexture) = 0;
+		/** Must be called called before TransferTexture on the same thread. */
+		virtual void ThreadPrep() = 0;
+		/** Must be called called after TransferTexture on the same thread. */
+		virtual void ThreadCleanup() = 0;
+
 		virtual void RegisterBuffer(const FRegisterDMABufferArgs& Args) = 0;
 		virtual void UnregisterBuffer(void* InBuffer) = 0;
 
diff --git a/Engine/Plugins/Media/MediaIOFramework/Source/MediaIOCore/Private/MediaCapture.cpp b/Engine/Plugins/Media/MediaIOFramework/Source/MediaIOCore/Private/MediaCapture.cpp
--- a/Engine/Plugins/Media/MediaIOFramework/Source/MediaIOCore/Private/MediaCapture.cpp
+++ b/Engine/Plugins/Media/MediaIOFramework/Source/MediaIOCore/Private/MediaCapture.cpp
@@ -493,7 +493,7 @@ bool UMediaCapture::UseExperimentalScheduling() const
 bool UMediaCapture::UseAnyThreadCapture() const
 {
 	return GRHISupportsMultithreading
-		&& CVarMediaIOEnableExperimentalScheduling.GetValueOnAnyThread() 
+		&& CVarMediaIOEnableExperimentalScheduling.GetValueOnAnyThread()
 		&& CVarMediaIOScheduleOnAnyThread.GetValueOnAnyThread()
 		&& SupportsAnyThreadCapture();
 }
@@ -1200,7 +1200,10 @@ void UMediaCapture::InitializeSyncHandlers_RenderThread()
 	for (int32 Index = 0; Index < NumberOfCaptureFrame; ++Index)
 	{
 		TSharedPtr<FMediaCaptureSyncData> SyncData = MakeShared<FMediaCaptureSyncData>();
-		SyncData->RHIFence = RHICreateGPUFence(*FString::Printf(TEXT("MediaCaptureSync_%02d"), Index));
+		if (!ShouldCaptureRHIResource())
+		{
+			SyncData->RHIFence = RHICreateGPUFence(*FString::Printf(TEXT("MediaCaptureSync_%02d"), Index));
+		}
 		SyncHandlers.Add(MoveTemp(SyncData));
 	}
 
@@ -1695,4 +1698,3 @@ namespace MediaCaptureDetails
 }
 
 #undef LOCTEXT_NAMESPACE
-
diff --git a/Engine/Plugins/Media/MediaIOFramework/Source/MediaIOCore/Private/MediaCaptureHelper.cpp b/Engine/Plugins/Media/MediaIOFramework/Source/MediaIOCore/Private/MediaCaptureHelper.cpp
--- a/Engine/Plugins/Media/MediaIOFramework/Source/MediaIOCore/Private/MediaCaptureHelper.cpp
+++ b/Engine/Plugins/Media/MediaIOFramework/Source/MediaIOCore/Private/MediaCaptureHelper.cpp
@@ -200,27 +200,38 @@ void FMediaCaptureHelper::ExecuteSyncPointPass(FRHICommandListImmediate& RHICmdL
 {
 	if (CapturingFrame && CapturingFrame->bMediaCaptureActive)
 	{
-		// Get available sync handler to create a sync point
-		TSharedPtr<UMediaCapture::FMediaCaptureSyncData> SyncDataPtr = MediaCapture->GetAvailableSyncHandler();
-		if (ensure(SyncDataPtr))
+		if (MediaCapture->ShouldCaptureRHIResource())
 		{
-			// This will happen after the conversion pass has completed
-			RHICmdList.WriteGPUFence(SyncDataPtr->RHIFence);
-			SyncDataPtr->bIsBusy = true;
+			MediaCapture->UnlockDMATexture_RenderThread(CapturingFrame->GetTextureResource());
 
-			// Queue capture data and our thread will wait (poll) and continue the process of providing a new texture
-			{
-				UE::MediaCaptureData::FPendingCaptureData NewCaptureData;
-				NewCaptureData.SyncHandler = SyncDataPtr;
-				NewCaptureData.CapturingFrame = CapturingFrame;
-				MediaCapture->SyncPointWatcher->QueuePendingCapture(MoveTemp(NewCaptureData));
-			}
+			UE::MediaCaptureData::FPendingCaptureData NewCaptureData;
+			NewCaptureData.CapturingFrame = CapturingFrame;
+			MediaCapture->SyncPointWatcher->QueuePendingCapture(MoveTemp(NewCaptureData));
 		}
 		else
 		{
-			UE_LOG(LogMediaIOCore, Error, TEXT(
-				"GetAvailableSyncHandler of MediaCapture '%s' failed to provide a fence, the captured buffers may not become available anymore."),
-				*MediaCapture->MediaOutputName);
+			// Get available sync handler to create a sync point
+			TSharedPtr<UMediaCapture::FMediaCaptureSyncData> SyncDataPtr = MediaCapture->GetAvailableSyncHandler();
+			if (ensure(SyncDataPtr))
+			{
+				// This will happen after the conversion pass has completed
+				RHICmdList.WriteGPUFence(SyncDataPtr->RHIFence);
+				SyncDataPtr->bIsBusy = true;
+
+				// Queue capture data and our thread will wait (poll) and continue the process of providing a new texture
+				{
+					UE::MediaCaptureData::FPendingCaptureData NewCaptureData;
+					NewCaptureData.SyncHandler = SyncDataPtr;
+					NewCaptureData.CapturingFrame = CapturingFrame;
+					MediaCapture->SyncPointWatcher->QueuePendingCapture(MoveTemp(NewCaptureData));
+				}
+			}
+			else
+			{
+				UE_LOG(LogMediaIOCore, Error, TEXT(
+					"GetAvailableSyncHandler of MediaCapture '%s' failed to provide a fence, the captured buffers may not become available anymore."),
+					*MediaCapture->MediaOutputName);
+			}
 		}
 	}
 }
@@ -309,11 +320,6 @@ void FMediaCaptureHelper::OnReadbackComplete(UMediaCapture* MediaCapture, TShare
 			{
 				if (ReadyFrame->IsTextureResource())
 				{
-					{
-						TRACE_CPUPROFILER_EVENT_SCOPE(UMediaCapture::UnlockDMATexture_RenderThread);
-						MediaCapture->UnlockDMATexture_RenderThread(ReadyFrame->GetTextureResource());
-					}
-
 					TRACE_CPUPROFILER_EVENT_SCOPE(UMediaCapture::RHIResourceCaptured);
 					TRACE_CPUPROFILER_EVENT_SCOPE_TEXT(*FString::Printf(TEXT("MediaCapture Output Frame %d"), ReadyFrame->CaptureBaseData.SourceFrameNumberRenderThread % 10));
 					SCOPE_CYCLE_COUNTER(STAT_MediaCaptureHelper_RenderThread_RHI_CaptureCallback)
@@ -364,4 +370,4 @@ void FMediaCaptureHelper::OnReadbackComplete(UMediaCapture* MediaCapture, TShare
 	};
 }
 
-}
\ No newline at end of file
+}
diff --git a/Engine/Plugins/Media/MediaIOFramework/Source/MediaIOCore/Private/MediaCaptureSyncPointWatcher.cpp b/Engine/Plugins/Media/MediaIOFramework/Source/MediaIOCore/Private/MediaCaptureSyncPointWatcher.cpp
--- a/Engine/Plugins/Media/MediaIOFramework/Source/MediaIOCore/Private/MediaCaptureSyncPointWatcher.cpp
+++ b/Engine/Plugins/Media/MediaIOFramework/Source/MediaIOCore/Private/MediaCaptureSyncPointWatcher.cpp
@@ -57,27 +57,43 @@ void FSyncPointWatcher::ProcessPendingCapture(const FPendingCaptureData& Data)
 	{
 		FScopedDurationTimer Timer(WaitTime);
 
-		// Wait until fence has been written (shader has completed)
-		while (true)
+		if (Data.SyncHandler)
 		{
-			if (Data.SyncHandler->RHIFence->Poll())
+			// Wait until fence has been written (shader has completed)
+			while (true)
 			{
-				bWaitedForCompletion = true;
-				break;
+				if (Data.SyncHandler->RHIFence->Poll())
+				{
+					bWaitedForCompletion = true;
+					break;
+				}
+
+				if (!Data.CapturingFrame->bMediaCaptureActive)
+				{
+					bWaitedForCompletion = false;
+					break;
+				}
+
+				constexpr float SleepTimeSeconds = 50 * 1E-6;
+				FPlatformProcess::SleepNoStats(SleepTimeSeconds);
 			}
 
+			Data.SyncHandler->RHIFence->Clear();
+			Data.SyncHandler->bIsBusy = false;
+		}
+		else
+		{
+			// This is a GPUDirect copy, so use ITextureTransfer's wait method to synchronize with the GPU.
 			if (!Data.CapturingFrame->bMediaCaptureActive)
 			{
 				bWaitedForCompletion = false;
-				break;
 			}
-
-			constexpr float SleepTimeSeconds = 50 * 1E-6;
-			FPlatformProcess::SleepNoStats(SleepTimeSeconds);
+			else
+			{
+				MediaCapture->WaitForGPU(Data.CapturingFrame->GetTextureResource());
+				bWaitedForCompletion = true;
+			}
 		}
-
-		Data.SyncHandler->RHIFence->Clear();
-		Data.SyncHandler->bIsBusy = false;
 	}
 
 	if (Data.CapturingFrame->bMediaCaptureActive && bWaitedForCompletion && MediaCapture)
@@ -143,4 +159,4 @@ FSyncPointWatcher::~FSyncPointWatcher()
 	WorkingThread.Reset();
 }
 
-}
\ No newline at end of file
+}
diff --git a/Engine/Plugins/Media/MediaIOFramework/Source/MediaIOCore/Public/MediaCapture.h b/Engine/Plugins/Media/MediaIOFramework/Source/MediaIOCore/Public/MediaCapture.h
--- a/Engine/Plugins/Media/MediaIOFramework/Source/MediaIOCore/Public/MediaCapture.h
+++ b/Engine/Plugins/Media/MediaIOFramework/Source/MediaIOCore/Public/MediaCapture.h
@@ -595,6 +595,9 @@ protected:
         return TexCreate_Shared | TexCreate_RenderTargetable | TexCreate_UAV;
     }
 
+	/** When using GPUTextureTransfer, calling this will wait for a signal set by the underlying library. */
+	virtual void WaitForGPU(FRHITexture* InRHITexture) {};
+
 protected:
 	UTextureRenderTarget2D* GetTextureRenderTarget() const;
 	TSharedPtr<FSceneViewport> GetCapturingSceneViewport() const;
