Fix static analysis warning in MediaCaptureSyncPointWatcher

#rnx

diff --git a/Engine/Plugins/Media/MediaIOFramework/Source/MediaIOCore/Private/MediaCaptureSyncPointWatcher.cpp b/Engine/Plugins/Media/MediaIOFramework/Source/MediaIOCore/Private/MediaCaptureSyncPointWatcher.cpp
--- a/Engine/Plugins/Media/MediaIOFramework/Source/MediaIOCore/Private/MediaCaptureSyncPointWatcher.cpp
+++ b/Engine/Plugins/Media/MediaIOFramework/Source/MediaIOCore/Private/MediaCaptureSyncPointWatcher.cpp
@@ -88,7 +88,7 @@ void FSyncPointWatcher::ProcessPendingCapture(const FPendingCaptureData& Data)
 			{
 				bWaitedForCompletion = false;
 			}
-			else
+			else if (MediaCapture)
 			{
 				MediaCapture->WaitForGPU(Data.CapturingFrame->GetTextureResource());
 				bWaitedForCompletion = true;
