Media IO - Add PostRender capture phase
Add groundwork to do some processing in the implementation before starting outputting.
Implementations can return false from IsReadyToOutput to indicate that the underlying SDK is not done initializing.
Preroll frames before starting to give frames to the underlying implementations, this should help ammortize the cost of initializing gpu resources before we output.

diff --git a/Engine/Plugins/Media/MediaIOFramework/Source/MediaIOCore/Private/MediaCapture.cpp b/Engine/Plugins/Media/MediaIOFramework/Source/MediaIOCore/Private/MediaCapture.cpp
--- a/Engine/Plugins/Media/MediaIOFramework/Source/MediaIOCore/Private/MediaCapture.cpp
+++ b/Engine/Plugins/Media/MediaIOFramework/Source/MediaIOCore/Private/MediaCapture.cpp
@@ -587,6 +587,7 @@ void UMediaCapture::StopCapture(bool bAllowPendingFrameToBeProcess)
 			SyncPointWatcher.Reset();
 			ViewExtension.Reset();
 			CaptureRenderPipeline.Reset();
+			bLoggedOutputNotReady = false;
 			
 			MediaOutputName.Reset();
 		}
@@ -618,7 +619,7 @@ bool UMediaCapture::SetCaptureAudioDevice(const FAudioDeviceHandle& InAudioDevic
 void UMediaCapture::CaptureImmediate_RenderThread(FRDGBuilder& GraphBuilder, FRHITexture* InSourceTexture, FIntRect SourceViewRect)
 {
 	bool bCaptureSucceeded = TryCaptureImmediate_RenderThread(GraphBuilder, InSourceTexture, SourceViewRect);
-	if(!bCaptureSucceeded)
+	if (!bCaptureSucceeded)
 	{
 		UE_LOG(LogMediaIOCore, VeryVerbose, TEXT("[%s] - Failed to capture resource."), *MediaOutputName);
 	}
@@ -626,25 +627,25 @@ void UMediaCapture::CaptureImmediate_RenderThread(FRDGBuilder& GraphBuilder, FRH
 
 bool UMediaCapture::TryCaptureImmediate_RenderThread(FRDGBuilder& GraphBuilder, FRHITexture* InSourceTexture, FIntRect SourceViewRect)
 {
-	UE::MediaCaptureData::FCaptureFrameArgs CaptureArgs{GraphBuilder};
+	UE::MediaCaptureData::FCaptureFrameArgs CaptureArgs{ GraphBuilder };
 	CaptureArgs.MediaCapture = this;
-    CaptureArgs.DesiredSize = DesiredSize;
-    CaptureArgs.ResourceToCapture = InSourceTexture;
+	CaptureArgs.DesiredSize = DesiredSize;
+	CaptureArgs.ResourceToCapture = InSourceTexture;
 	CaptureArgs.SourceViewRect = SourceViewRect;
-	
+
 	return CaptureImmediate_RenderThread(CaptureArgs);
 }
 
-void UMediaCapture::CaptureImmediate_RenderThread(FRDGBuilder& GraphBuilder,  FRDGTextureRef InSourceTextureRef, FIntRect SourceViewRect)
+void UMediaCapture::CaptureImmediate_RenderThread(FRDGBuilder& GraphBuilder, FRDGTextureRef InSourceTextureRef, FIntRect SourceViewRect)
 {
 	bool bCaptureSucceeded = TryCaptureImmediate_RenderThread(GraphBuilder, InSourceTextureRef, SourceViewRect);
-	if(!bCaptureSucceeded)
+	if (!bCaptureSucceeded)
 	{
 		UE_LOG(LogMediaIOCore, VeryVerbose, TEXT("[%s] - Failed to capture resource."), *MediaOutputName);
 	}
 }
 
-bool UMediaCapture::TryCaptureImmediate_RenderThread(FRDGBuilder& GraphBuilder,  FRDGTextureRef InSourceTextureRef, FIntRect SourceViewRect)
+bool UMediaCapture::TryCaptureImmediate_RenderThread(FRDGBuilder& GraphBuilder, FRDGTextureRef InSourceTextureRef, FIntRect SourceViewRect)
 {
 	// Make sure captureimmediate runs are seen by game thread when stopping is requested. Having a pending render command will cause a flush
 	// hence letting the capture finish. If stopping is in progress, stopped state will be seen by the captureimmediate and early exit.
@@ -654,13 +655,13 @@ bool UMediaCapture::TryCaptureImmediate_RenderThread(FRDGBuilder& GraphBuilder,
 	{
 		--WaitingForRenderCommandExecutionCounter;
 	};
-	
+
 	UE::MediaCaptureData::FCaptureFrameArgs CaptureArgs{GraphBuilder};
 	CaptureArgs.MediaCapture = this;
 	CaptureArgs.DesiredSize = DesiredSize;
 	CaptureArgs.RDGResourceToCapture = InSourceTextureRef;
 	CaptureArgs.SourceViewRect = SourceViewRect;
-	
+
 	return CaptureImmediate_RenderThread(CaptureArgs);
 }
 
@@ -669,9 +670,9 @@ bool UMediaCapture::CaptureImmediate_RenderThread(const UE::MediaCaptureData::FC
 	using namespace UE::MediaCaptureData;
 	TRACE_CPUPROFILER_EVENT_SCOPE(UMediaCapture::CaptureImmediate_RenderThread);
 	TRACE_CPUPROFILER_EVENT_SCOPE_TEXT(*FString::Printf(TEXT("MediaCapturePipe: %llu"), GFrameCounterRenderThread % 10));
-	
+
 	check(IsInRenderingThread());
-	
+
 	if (bIsAutoRestartRequired)
 	{
 		return false;
@@ -720,7 +721,7 @@ bool UMediaCapture::CaptureImmediate_RenderThread(const UE::MediaCaptureData::FC
 	if (CaptureSource->GetSourceType() == EMediaCaptureSourceType::RHI_RESOURCE)
 	{
 		// Keep resource size up to date with incoming resource to capture
- 		StaticCastSharedPtr<UE::MediaCapture::Private::FRHIResourceCaptureSource>(CaptureSource)->ResourceDescription.ResourceSize = Args.GetSizeXY();
+		StaticCastSharedPtr<UE::MediaCapture::Private::FRHIResourceCaptureSource>(CaptureSource)->ResourceDescription.ResourceSize = Args.GetSizeXY();
 	}
 
 	// Get cached capture data from game thread. We want to find a cached frame matching current render thread frame number
@@ -739,7 +740,7 @@ bool UMediaCapture::CaptureImmediate_RenderThread(const UE::MediaCaptureData::FC
 				break;
 			}
 			else if (GFrameCounterRenderThread > It->BaseData.SourceFrameNumberRenderThread &&
-					GFrameCounterRenderThread - It->BaseData.SourceFrameNumberRenderThread > MaxCaptureDataAgeInFrames)
+				GFrameCounterRenderThread - It->BaseData.SourceFrameNumberRenderThread > MaxCaptureDataAgeInFrames)
 			{
 				// Remove old frame data that wasn't used.
 				It.RemoveCurrent();
@@ -758,7 +759,7 @@ bool UMediaCapture::CaptureImmediate_RenderThread(const UE::MediaCaptureData::FC
 		StopCapture(false);
 		return false;
 	}
-		
+
 	// Get next available frame from the store. Can be invalid.
 	TSharedPtr<FCaptureFrame> CapturingFrame;
 	if (GetState() != EMediaCaptureState::StopRequested)
@@ -776,7 +777,7 @@ bool UMediaCapture::CaptureImmediate_RenderThread(const UE::MediaCaptureData::FC
 
 			if (!CapturingFrame)
 			{
-				UE_LOG(LogMediaIOCore, Error, TEXT("[%s] - No frames available for capture for frame %llu. This should not happen."), 
+				UE_LOG(LogMediaIOCore, Error, TEXT("[%s] - No frames available for capture for frame %llu. This should not happen."),
 					*MediaOutputName, GFrameCounterRenderThread);
 
 				SetState(EMediaCaptureState::Error);
@@ -786,7 +787,7 @@ bool UMediaCapture::CaptureImmediate_RenderThread(const UE::MediaCaptureData::FC
 		else
 		{
 			//In case we are skipping frames, just keep capture frame as invalid
-			UE_LOG(LogMediaIOCore, Warning, TEXT("[%s] - No frames available for capture of frame %llu. Skipping"), 
+			UE_LOG(LogMediaIOCore, Warning, TEXT("[%s] - No frames available for capture of frame %llu. Skipping"),
 				*MediaOutputName, GFrameCounterRenderThread);
 		}
 	}
@@ -797,9 +798,9 @@ bool UMediaCapture::CaptureImmediate_RenderThread(const UE::MediaCaptureData::FC
 	{
 		TRACE_CPUPROFILER_EVENT_SCOPE(UMediaCapture::Capture_Invalid);
 	}
-	
+
 	UE_LOG(LogMediaIOCore, Verbose, TEXT("[%s - %s] - Capturing frame %d"), *MediaOutputName, *FThreadManager::GetThreadName(FPlatformTLS::GetCurrentThreadId()), CapturingFrame ? CapturingFrame->FrameId : -1);
-	
+
 	bool bCaptureSucceeded = false;
 	if (CapturingFrame)
 	{
@@ -1022,11 +1023,12 @@ void UMediaCapture::ProcessCapture_GameThread()
 			UE_LOG(LogMediaIOCore, VeryVerbose, TEXT("[%s] - Flushing commands."), *MediaOutputName);
 
 			FlushRenderingCommands();
-			
+
 			WaitForSingleExperimentalSchedulingTaskToComplete();
 
 			// After flushing, we should have access to an available frame, if not, it's not expected.
 			CapturingFrame = FrameManager->GetNextAvailable<FCaptureFrame>();
+
 			if (!CapturingFrame)
 			{
 				UE_LOG(LogMediaIOCore, Error, TEXT("[%s] - Flushing commands didn't give us back available frames to process.") , *MediaOutputName);
@@ -1085,7 +1087,7 @@ void UMediaCapture::PrepareAndDispatchCapture_GameThread(const TSharedPtr<UE::Me
 				{
 					FRDGBuilder GraphBuilder(RHICmdList);
 					
-					FTextureRHIRef SourceTexture = InMediaCapture->CaptureSource->GetSourceTextureForInput_RenderThread(RHICmdList);
+					FTextureRHIRef SourceTexture = InMediaCapture->CaptureSource->GetSourceTextureForInput_RenderThread();
 					
 					FCaptureFrameArgs CaptureArgs{ GraphBuilder };
 					CaptureArgs.MediaCapture = InMediaCapture;
@@ -1113,7 +1115,7 @@ void UMediaCapture::PrepareAndDispatchCapture_GameThread(const TSharedPtr<UE::Me
 			});
 
 			//If auto-stopping, count the number of frame captures requested and stop when reaching 0.
-			if (DesiredCaptureOptions.bAutostopOnCapture && GetState() == EMediaCaptureState::Capturing && --DesiredCaptureOptions.NumberOfFramesToCapture <= 0)
+			if (DesiredCaptureOptions.bAutostopOnCapture && GetState() == EMediaCaptureState::Capturing && CaptureRequestCount == DesiredCaptureOptions.NumberOfFramesToCapture)
 			{
 				StopCapture(true);
 			}
@@ -1167,6 +1169,12 @@ void UMediaCapture::InitializeOutputResources(int32 InNumberOfBuffers)
 						NewFrame->ReadbackTexture = MakeUnique<FRHIGPUTextureReadback>(*FString::Printf(TEXT("MediaCaptureTextureReadback_%d"), Index));
 					}
 
+					if (This->GetDesiredCaptureOptions().bAutostopOnCapture)
+					{
+						// Don't preroll if capturing a set number of frames since we discard the first frames.
+						NewFrame->bPrerollFrame = false;
+					}
+
 					This->FrameManager->AddFrame(MoveTemp(NewFrame));
 				}
 				else
@@ -1181,6 +1189,12 @@ void UMediaCapture::InitializeOutputResources(int32 InNumberOfBuffers)
 							NewFrame->ReadbackBuffer = MakeUnique<FRHIGPUBufferReadback>(*FString::Printf(TEXT("MediaCaptureBufferReadback_%d"), Index));
 						}
 
+						if (This->GetDesiredCaptureOptions().bAutostopOnCapture)
+						{
+							// Don't preroll if capturing a set number of frames since we discard the first frames.
+							NewFrame->bPrerollFrame = false;
+						}
+
 						This->FrameManager->AddFrame(MoveTemp(NewFrame));
 					}
 					else
@@ -1301,7 +1315,10 @@ void UMediaCapture::OnBeginFrame_GameThread()
 		CaptureData.UserData = GetCaptureFrameUserData_GameThread();
 
 		FScopeLock Lock(&CaptureDataQueueCriticalSection);
+
+		TRACE_BOOKMARK(TEXT("MediaCapture Push frame %d"), GFrameCounter);
 		CaptureDataQueue.Insert(MoveTemp(CaptureData), 0);
+
 	}
 }
 
@@ -1355,9 +1372,6 @@ void UMediaCapture::OnEndFrame_GameThread()
 bool UMediaCapture::ProcessCapture_RenderThread(const TSharedPtr<UE::MediaCaptureData::FCaptureFrame>& CapturingFrame, UE::MediaCaptureData::FCaptureFrameArgs Args)
 {
 	using namespace UE::MediaCaptureData;
-	
-	RDG_EVENT_SCOPE_STAT(Args.GraphBuilder, MediaCapture_ProcessCapture, "MediaCapture_ProcessCapture");
-	RDG_GPU_STAT_SCOPE(Args.GraphBuilder, MediaCapture_ProcessCapture);
 
 	TRACE_CPUPROFILER_EVENT_SCOPE(UMediaCapture::ProcessCapture_RenderThread);
 
@@ -1366,6 +1380,10 @@ bool UMediaCapture::ProcessCapture_RenderThread(const TSharedPtr<UE::MediaCaptur
 	{
 		FrameNumber = CapturingFrame->CaptureBaseData.SourceFrameNumberRenderThread;
 	}
+
+	RDG_EVENT_SCOPE_STAT(Args.GraphBuilder, MediaCapture_ProcessCapture, "MediaCapture_ProcessCapture", FrameNumber % 10);
+	RDG_GPU_STAT_SCOPE(Args.GraphBuilder, MediaCapture_ProcessCapture);
+
 	TRACE_CPUPROFILER_EVENT_SCOPE_TEXT(*FString::Printf(TEXT("Process Capture Render Thread Frame %d"), FrameNumber % 10));
 
 	if (CapturingFrame)
@@ -1382,37 +1400,63 @@ bool UMediaCapture::ProcessCapture_RenderThread(const TSharedPtr<UE::MediaCaptur
 				Args.RDGResourceToCapture = Args.GraphBuilder.RegisterExternalTexture(CreateRenderTarget(Args.ResourceToCapture, TEXT("MediaCaptureSourceTexture")));
 			}
 
-			Args.MediaCapture->CaptureRenderPipeline->InitializeResources_RenderThread(CapturingFrame, Args.GraphBuilder, Args.RDGResourceToCapture);
+			if (CapturingFrame->bPrerollFrame)
+			{
+				// Make sure we're not doing prerolling when capturing a certain number of frames.
+				check(!GetDesiredCaptureOptions().bAutostopOnCapture);
+			}
 
-			// Call the capture frame algo based on the specific type of resource we are using
-			if (DesiredOutputResourceType == EMediaCaptureResourceType::Texture)
+			if (CapturingFrame->bPrerollFrame || GetDesiredCaptureOptions().bAutostopOnCapture)
 			{
-				if (ensure(CapturingFrame->IsTextureResource()))
-				{
-					Args.MediaCapture->BeforeFrameCaptured_RenderingThread(CapturingFrame->CaptureBaseData, CapturingFrame->UserData, CapturingFrame->GetTextureResource());
-				}
-				else
-				{
-					bBeforeCaptureSuccess = false;
-					UE_LOG(LogMediaIOCore, Error, TEXT("The capture will stop for '%s'. Capture frame was expected to use Texture resource but wasn't."), *Args.MediaCapture->MediaOutputName);
-				}
+				Args.MediaCapture->CaptureRenderPipeline->InitializeResources_RenderThread(CapturingFrame, Args.GraphBuilder, Args.RDGResourceToCapture);
 			}
-			else
+
+			const bool bIsReadyToStartOutput = Args.MediaCapture->IsReadyToOutput();
+
+			if (!CapturingFrame->bPrerollFrame && !bIsReadyToStartOutput && !bLoggedOutputNotReady)
+			{
+				bLoggedOutputNotReady = true;
+				UE_LOG(LogMediaIOCore, Display, TEXT("Skipping frame since capture implementation is not ready to output yet."));
+			}
+
+			if (!CapturingFrame->bPrerollFrame && bIsReadyToStartOutput)
 			{
-				if (ensure(CapturingFrame->IsBufferResource()))
+				// Call the capture frame algo based on the specific type of resource we are using
+				if (DesiredOutputResourceType == EMediaCaptureResourceType::Texture)
 				{
-					Args.MediaCapture->BeforeFrameCaptured_RenderingThread(CapturingFrame->CaptureBaseData, CapturingFrame->UserData, CapturingFrame->GetBufferResource());
+					if (ensure(CapturingFrame->IsTextureResource()))
+					{
+						Args.MediaCapture->BeforeFrameCaptured_RenderingThread(CapturingFrame->CaptureBaseData, CapturingFrame->UserData, CapturingFrame->GetTextureResource());
+					}
+					else
+					{
+						bBeforeCaptureSuccess = false;
+						UE_LOG(LogMediaIOCore, Error, TEXT("The capture will stop for '%s'. Capture frame was expected to use Texture resource but wasn't."), *Args.MediaCapture->MediaOutputName);
+					}
 				}
 				else
 				{
-					bBeforeCaptureSuccess = false;
-					UE_LOG(LogMediaIOCore, Error, TEXT("The capture will stop for '%s'. Capture frame was expected to use Buffer resource but wasn't."), *Args.MediaCapture->MediaOutputName);
+					if (ensure(CapturingFrame->IsBufferResource()))
+					{
+						Args.MediaCapture->BeforeFrameCaptured_RenderingThread(CapturingFrame->CaptureBaseData, CapturingFrame->UserData, CapturingFrame->GetBufferResource());
+					}
+					else
+					{
+						bBeforeCaptureSuccess = false;
+						UE_LOG(LogMediaIOCore, Error, TEXT("The capture will stop for '%s'. Capture frame was expected to use Buffer resource but wasn't."), *Args.MediaCapture->MediaOutputName);
+					}
 				}
-			}
 
-			if (bBeforeCaptureSuccess)
+				if (bBeforeCaptureSuccess)
+				{
+					bHasCaptureSuceeded = UE::MediaCaptureData::FMediaCaptureHelper::CaptureFrame(Args, CapturingFrame);
+				}
+			}
+			else
 			{
-				bHasCaptureSuceeded = UE::MediaCaptureData::FMediaCaptureHelper::CaptureFrame(Args, CapturingFrame);
+				bHasCaptureSuceeded = true;
+				CapturingFrame->bPrerollFrame = false;
+				FrameManager->MarkAvailable(*CapturingFrame);
 			}
 		}
 		else
diff --git a/Engine/Plugins/Media/MediaIOFramework/Source/MediaIOCore/Private/MediaCaptureSceneViewExtension.h b/Engine/Plugins/Media/MediaIOFramework/Source/MediaIOCore/Private/MediaCaptureSceneViewExtension.h
--- a/Engine/Plugins/Media/MediaIOFramework/Source/MediaIOCore/Private/MediaCaptureSceneViewExtension.h
+++ b/Engine/Plugins/Media/MediaIOFramework/Source/MediaIOCore/Private/MediaCaptureSceneViewExtension.h
@@ -86,6 +86,22 @@ public:
 	{
 		return true;
 	}
+
+	virtual void PostRenderView_RenderThread(FRDGBuilder& GraphBuilder, FSceneView& InView) override
+	{
+		if (CapturePhase == EMediaCapturePhase::PostRender)
+		{
+			TRACE_CPUPROFILER_EVENT_SCOPE(MediaCaptureExtensionCallback);
+
+			FTextureRHIRef RenderTarget = InView.Family->RenderTarget->GetRenderTargetTexture();
+
+			RDG_EVENT_SCOPE(GraphBuilder, "MediaCaptureSceneExtension");
+			if (WeakCapture.IsValid())
+			{
+				WeakCapture->TryCaptureImmediate_RenderThread(GraphBuilder, RenderTarget, InView.UnscaledViewRect);
+			}
+		}
+	}
 	//~ End FSceneViewExtensionBase Interface
 
 	FScreenPassTexture PostProcessCallback_RenderThread(FRDGBuilder& GraphBuilder, const FSceneView& View, const FPostProcessMaterialInputs& InOutInputs)
diff --git a/Engine/Plugins/Media/MediaIOFramework/Source/MediaIOCore/Private/MediaCaptureSources.h b/Engine/Plugins/Media/MediaIOFramework/Source/MediaIOCore/Private/MediaCaptureSources.h
--- a/Engine/Plugins/Media/MediaIOFramework/Source/MediaIOCore/Private/MediaCaptureSources.h
+++ b/Engine/Plugins/Media/MediaIOFramework/Source/MediaIOCore/Private/MediaCaptureSources.h
@@ -167,7 +167,7 @@ namespace UE::MediaCapture::Private
 		}
 		
 		/** Only returns a valid texture for scene viewports and render targets as the RHI Sources don't keep a pointer to the RHI texture. */
-		virtual FTextureRHIRef GetSourceTextureForInput_RenderThread(FRHICommandListImmediate& RHICmdList)
+		virtual FTextureRHIRef GetSourceTextureForInput_RenderThread()
 		{
 			return nullptr;
 		}
@@ -278,7 +278,7 @@ namespace UE::MediaCapture::Private
 			return false;
 		}
 
-		virtual FTextureRHIRef GetSourceTextureForInput_RenderThread(FRHICommandListImmediate& RHICmdList) override
+		virtual FTextureRHIRef GetSourceTextureForInput_RenderThread() override
 		{
 			FTextureRHIRef SourceTexture;
 			if (TSharedPtr<FSceneViewport> Viewport = WeakViewport.Pin())
@@ -402,7 +402,7 @@ namespace UE::MediaCapture::Private
 			return false;
 		}
 
-		virtual FTextureRHIRef GetSourceTextureForInput_RenderThread(FRHICommandListImmediate& RHICmdList)
+		virtual FTextureRHIRef GetSourceTextureForInput_RenderThread()
 		{
 			constexpr bool bEvenIfPendingKill = false;
 			constexpr bool bThreadSafeTest = true;
diff --git a/Engine/Plugins/Media/MediaIOFramework/Source/MediaIOCore/Private/MediaIOFrameManager.h b/Engine/Plugins/Media/MediaIOFramework/Source/MediaIOCore/Private/MediaIOFrameManager.h
--- a/Engine/Plugins/Media/MediaIOFramework/Source/MediaIOCore/Private/MediaIOFrameManager.h
+++ b/Engine/Plugins/Media/MediaIOFramework/Source/MediaIOCore/Private/MediaIOFrameManager.h
@@ -287,6 +287,8 @@ namespace UE::MediaCaptureData
 		TSharedPtr<FMediaCaptureUserData> UserData;
 		/** Resources used by the media capture render pipeline (ie. Color conversion render target ) */
 		UE::MediaCapture::FRenderPassFrameResources RenderPassResources;
+		/** Whether this frame is used to preroll the output. Preroll frames are only used to create render resources to avoid a startup hitch. */
+		bool bPrerollFrame = true;
 	};
 
 	/** Parameter to make our sync pass needing the convert pass as a prereq */
diff --git a/Engine/Plugins/Media/MediaIOFramework/Source/MediaIOCore/Public/MediaCapture.h b/Engine/Plugins/Media/MediaIOFramework/Source/MediaIOCore/Public/MediaCapture.h
--- a/Engine/Plugins/Media/MediaIOFramework/Source/MediaIOCore/Public/MediaCapture.h
+++ b/Engine/Plugins/Media/MediaIOFramework/Source/MediaIOCore/Public/MediaCapture.h
@@ -142,6 +142,7 @@ enum class EMediaCapturePhase : uint8
 	AfterMotionBlur,
 	AfterToneMap,
 	AfterFXAA,
+	PostRender,
 	EndFrame
 };
 
@@ -436,6 +437,9 @@ public:
 	virtual FString GetDesc() override;
 
 protected:
+	/** Implementations can return false when the output sdk is not ready. When this function returns false, the current output frame will be discarded. */
+	virtual bool IsReadyToOutput() const { return true; }
+
 	//~ UMediaCapture interface
 	virtual bool ValidateMediaOutput() const;
 
@@ -682,9 +686,10 @@ protected:
 	/** Type of capture currently being done. i.e Viewport, RT, RHI Resource */
 	EMediaCaptureSourceType CaptureSourceType = EMediaCaptureSourceType::SCENE_VIEWPORT;
 
+	int32 NumberOfCaptureFrame = 2;
+
 private:	
 	TPimplPtr<UE::MediaCaptureData::FFrameManager> FrameManager;
-	int32 NumberOfCaptureFrame = 2;
 	int32 CaptureRequestCount = 0;
 	std::atomic<EMediaCaptureState> MediaState = EMediaCaptureState::Stopped;
 
@@ -702,6 +707,8 @@ private:
 	EMediaCaptureResourceType DesiredOutputResourceType = EMediaCaptureResourceType::Texture;
 	FString MediaOutputName = FString(TEXT("[undefined]"));
 	bool bUseRequestedTargetSize = false;
+	/** Used to avoid repeatedly logging while the underlying sdk is not ready to output. */
+	bool bLoggedOutputNotReady = false;
 
 	bool bViewportHasFixedViewportSize = false;
 	std::atomic<bool> bOutputResourcesInitialized;
