[MediaIOCore] Audio - Fixing the thread safety issues around buffer listener and delegates when adding or removing audio out.

Issues occur when media capture has an audio device explicitly set, which leads to FMediaIOAudioCapture being created and destroyed.

1- Can't call AsShared() from ~FMediaIOAudioCapture

UnregisterSubmixBufferListener requires a SharedRef to the buffer listener. This being called from ~FMediaIOAudioCapture means it can't be the buffer listener itself.
We need to make it into a sub-object that remains a valid shared ref during destruction.

2- Delegates FMediaIOAudioCapture::AudioCapturedDelegate and UMediaIOCoreSubsystem::BufferReceivedDelegate need to be thread safe

We can make UMediaIOCoreSubsystem::BufferReceivedDelegate thread safe so it can be modified while audio thread calls it.

There is no thread safe version for FMediaIOAudioCapture::AudioCapturedDelegate. Moving the callback setup before Initialize removes the possiblity of audio thread calling it.
If nothing else modifies this delegate, there shouldn't be a need to protect it further.

diff --git a/Engine/Plugins/Media/MediaIOFramework/Source/MediaIOCore/Private/MediaIOCoreAudioOutput.cpp b/Engine/Plugins/Media/MediaIOFramework/Source/MediaIOCore/Private/MediaIOCoreAudioOutput.cpp
--- a/Engine/Plugins/Media/MediaIOFramework/Source/MediaIOCore/Private/MediaIOCoreAudioOutput.cpp
+++ b/Engine/Plugins/Media/MediaIOFramework/Source/MediaIOCore/Private/MediaIOCoreAudioOutput.cpp
@@ -1,4 +1,4 @@
-﻿// Copyright Epic Games, Inc. All Rights Reserved.
+// Copyright Epic Games, Inc. All Rights Reserved.
 
 #include "MediaIOCoreAudioOutput.h"
 
@@ -85,11 +85,45 @@ Audio::FAlignedFloatBuffer FMediaIOAudioOutput::GetFloatBuffer(uint32 NumSamples
 	return FloatBuffer;
 }
 
-PRAGMA_DISABLE_DEPRECATION_WARNINGS
-FMediaIOAudioCapture::FMediaIOAudioCapture(const FAudioDeviceHandle& InAudioDeviceHandle)
+/**
+ * Implementation of ISubmixBufferListener.
+ * 
+ * It is necessary to make it an aggregate so it remains valid as SharedRef from within ~FMediaIOAudioCapture().
+ * This is because UnregisterSubmixBufferListener requires a valid SharedRef, which AsShared() can't provide
+ * from within the destructor.
+ */ 
+class FMediaIOAudioCapture::FSubmixBufferListenerProxy : public ISubmixBufferListener
 {
-}
-PRAGMA_ENABLE_DEPRECATION_WARNINGS
+public:
+	FSubmixBufferListenerProxy(const TSharedRef<ISubmixBufferListener>& InSubmixBufferListener)
+		: SubmixBufferListenerWeak(InSubmixBufferListener)
+	{
+	}
+
+	//~ Begin ISubmixBufferListener
+	virtual void OnNewSubmixBuffer(const USoundSubmix* InOwningSubmix, float* InAudioData, int32 InNumSamples, int32 InNumChannels, const int32 InSampleRate, double InAudioClock) override
+	{
+		if (const TSharedPtr<ISubmixBufferListener> SubmixBufferListener = SubmixBufferListenerWeak.Pin())
+		{
+			SubmixBufferListener->OnNewSubmixBuffer(InOwningSubmix, InAudioData, InNumSamples, InNumChannels, InSampleRate, InAudioClock);
+		}
+	}
+
+	virtual const FString& GetListenerName() const override
+	{
+		if (const TSharedPtr<ISubmixBufferListener> SubmixBufferListener = SubmixBufferListenerWeak.Pin())
+		{
+			return SubmixBufferListener->GetListenerName();
+		}
+
+		static const FString ListenerName = TEXT("MediaIO AudioCapture Listener Bridge");
+		return ListenerName;
+	}
+	//~ End ISubmixBufferListener
+
+private:
+	TWeakPtr<ISubmixBufferListener> SubmixBufferListenerWeak;
+};
 
 FMediaIOAudioCapture::~FMediaIOAudioCapture()
 {
@@ -101,6 +135,14 @@ void FMediaIOAudioCapture::Initialize(const FAudioDeviceHandle& InAudioDeviceHan
 	RegisterAudioDevice(InAudioDeviceHandle);
 }
 
+void FMediaIOAudioCapture::Initialize(const FAudioDeviceHandle& InAudioDeviceHandle, const FOnAudioCaptured& InAudioCapturedDelegate)
+{
+	// Important: Assign the delegate before registering the buffer listener to avoid thread safety issues.
+	AudioCapturedDelegate = InAudioCapturedDelegate;
+
+	RegisterAudioDevice(InAudioDeviceHandle);
+}
+
 void FMediaIOAudioCapture::OnNewSubmixBuffer(const USoundSubmix* InOwningSubmix, float* InAudioData, int32 InNumSamples, int32 InNumChannels, const int32 InSampleRate, double InAudioClock)
 {
 	TRACE_CPUPROFILER_EVENT_SCOPE(FMediaIOAudioCapture::OnNewSubmixBuffer);
@@ -111,8 +153,7 @@ void FMediaIOAudioCapture::OnNewSubmixBuffer(const USoundSubmix* InOwningSubmix,
 		if (ensureMsgf(NumChannels == InNumChannels, TEXT("Expected %d channels from submix buffer but got %d instead."), NumChannels, InNumChannels))
 		{
 			AudioCapturedDelegate.ExecuteIfBound(InAudioData, InNumSamples);
-			
-			int32 NumPushed = AudioSplitter.PushAudio(InAudioData, InNumSamples);
+			const int32 NumPushed = AudioSplitter.PushAudio(InAudioData, InNumSamples);
 			if (InNumSamples != NumPushed && NumPushed != -1)
 			{
 				TRACE_CPUPROFILER_EVENT_SCOPE(FMediaIOAudioCapture::OnNewSubmixBuffer::BufferOverrun);
@@ -169,7 +210,12 @@ void FMediaIOAudioCapture::RegisterBufferListener(FAudioDevice* AudioDevice)
 		NumChannels = MixerDevice->GetDeviceOutputChannels();
 		SampleRate = MixerDevice->GetSampleRate();
 		PrimarySubmixName = *GetDefault<UAudioSettings>()->MasterSubmix.GetAssetName();
-		AudioDevice->RegisterSubmixBufferListener(AsShared(), AudioDevice->GetMainSubmixObject());
+
+		if (!SubmixBufferListenerProxy)
+		{
+			SubmixBufferListenerProxy = MakeShared<FSubmixBufferListenerProxy>(AsShared());
+		}
+		AudioDevice->RegisterSubmixBufferListener(SubmixBufferListenerProxy.ToSharedRef(), AudioDevice->GetMainSubmixObject());
 	}
 }
 
@@ -181,9 +227,9 @@ const FString& FMediaIOAudioCapture::GetListenerName() const
 
 void FMediaIOAudioCapture::UnregisterBufferListener(FAudioDevice* AudioDevice)
 {
-	if (AudioDevice)
+	if (AudioDevice && SubmixBufferListenerProxy)
 	{
-		AudioDevice->UnregisterSubmixBufferListener(AsShared(), AudioDevice->GetMainSubmixObject());
+		AudioDevice->UnregisterSubmixBufferListener(SubmixBufferListenerProxy.ToSharedRef(), AudioDevice->GetMainSubmixObject());
 	}
 }
 
@@ -211,6 +257,16 @@ void FMediaIOAudioCapture::UnregisterAudioDevice()
 	}
 }
 
+void FMediaIOAudioCapture::SetAudioCapturedDelegate(const FOnAudioCaptured& InAudioCapturedDelegate)
+{
+	// We need to protect against thread unsafe operation. The audio capture delegate can't be
+	// changed if the SubmixBufferListener is registered to an audio device.
+	if (ensure(!SubmixBufferListenerProxy || SubmixBufferListenerProxy.GetSharedReferenceCount() == 1))
+	{
+		AudioCapturedDelegate = InAudioCapturedDelegate;
+	}
+}
+
 FMainMediaIOAudioCapture::FMainMediaIOAudioCapture()
 {	
 #if WITH_EDITOR
@@ -233,6 +289,15 @@ void FMainMediaIOAudioCapture::Initialize()
 	RegisterCurrentAudioDevice();
 }
 
+void FMainMediaIOAudioCapture::Initialize(const FOnAudioCaptured& InAudioCapturedDelegate)
+{
+	// Important: Assign the delegate before registering the buffer listener to avoid thread safety issues.
+	SetAudioCapturedDelegate(InAudioCapturedDelegate);
+
+	// Register the current device (PIE or main).
+	RegisterCurrentAudioDevice();
+}
+
 #if WITH_EDITOR
 void FMainMediaIOAudioCapture::OnPIEStarted(const bool)
 {
diff --git a/Engine/Plugins/Media/MediaIOFramework/Source/MediaIOCore/Private/MediaIOCoreSubsystem.cpp b/Engine/Plugins/Media/MediaIOFramework/Source/MediaIOCore/Private/MediaIOCoreSubsystem.cpp
--- a/Engine/Plugins/Media/MediaIOFramework/Source/MediaIOCore/Private/MediaIOCoreSubsystem.cpp
+++ b/Engine/Plugins/Media/MediaIOFramework/Source/MediaIOCore/Private/MediaIOCoreSubsystem.cpp
@@ -18,6 +18,7 @@ void UMediaIOCoreSubsystem::Deinitialize()
 TSharedPtr<FMediaIOAudioOutput> UMediaIOCoreSubsystem::CreateAudioOutput(const FCreateAudioOutputArgs& InArgs)
 {
 	TSharedPtr<FMediaIOAudioCapture> MediaIOAudioCapture;
+	using FOnAudioCaptured = FMediaIOAudioCapture::FOnAudioCaptured;
 
 	if (InArgs.AudioDeviceHandle.IsValid())
 	{
@@ -27,9 +28,9 @@ TSharedPtr<FMediaIOAudioOutput> UMediaIOCoreSubsystem::CreateAudioOutput(const F
 		}
 		else
 		{
+			const FOnAudioCaptured OnAudioCapturedDelegate = FOnAudioCaptured::CreateUObject(this, &UMediaIOCoreSubsystem::OnBufferReceivedByCapture, InArgs.AudioDeviceHandle.GetDeviceID());
 			MediaIOAudioCapture = MediaIOAudioCaptures.Add(InArgs.AudioDeviceHandle.GetDeviceID(), MakeShared<FMediaIOAudioCapture>());
-			MediaIOAudioCapture->Initialize(InArgs.AudioDeviceHandle);
-			MediaIOAudioCapture->OnAudioCaptured_RenderThread().BindUObject(this, &UMediaIOCoreSubsystem::OnBufferReceivedByCapture, InArgs.AudioDeviceHandle.GetDeviceID());
+			MediaIOAudioCapture->Initialize(InArgs.AudioDeviceHandle, OnAudioCapturedDelegate);
 		}
 	}
 	else
@@ -37,9 +38,9 @@ TSharedPtr<FMediaIOAudioOutput> UMediaIOCoreSubsystem::CreateAudioOutput(const F
 		// Fallback using the main audio device.
 		if (!MainMediaIOAudioCapture)
 		{
+			const FOnAudioCaptured OnAudioCapturedDelegate = FOnAudioCaptured::CreateUObject(this, &UMediaIOCoreSubsystem::OnBufferReceivedByCapture, InArgs.AudioDeviceHandle.GetDeviceID());
 			MainMediaIOAudioCapture = MakeShared<FMainMediaIOAudioCapture>();
-			MainMediaIOAudioCapture->Initialize();
-			MainMediaIOAudioCapture->OnAudioCaptured_RenderThread().BindUObject(this, &UMediaIOCoreSubsystem::OnBufferReceivedByCapture, InArgs.AudioDeviceHandle.GetDeviceID());
+			MainMediaIOAudioCapture->Initialize(OnAudioCapturedDelegate);
 		}
 		MediaIOAudioCapture = MainMediaIOAudioCapture;
 	}	
diff --git a/Engine/Plugins/Media/MediaIOFramework/Source/MediaIOCore/Public/MediaIOCoreAudioOutput.h b/Engine/Plugins/Media/MediaIOFramework/Source/MediaIOCore/Public/MediaIOCoreAudioOutput.h
--- a/Engine/Plugins/Media/MediaIOFramework/Source/MediaIOCore/Public/MediaIOCoreAudioOutput.h
+++ b/Engine/Plugins/Media/MediaIOFramework/Source/MediaIOCore/Public/MediaIOCoreAudioOutput.h
@@ -109,6 +109,7 @@ public:
 		return ConvertAndUpmixBuffer<OutputType>(FloatBuffer, NumInputChannels, NumOutputChannels);
 	}
 
+	UE_DEPRECATED(5.6, "This function is deprecated because it is not implemented. Use UMediaIOCoreSubsystem::OnBufferReceived() instead.")
 	FOnBufferReceived& OnBufferReceived()
 	{
 		return BufferReceivedDelegate;
@@ -157,7 +158,7 @@ private:
 	/** The buffer accumulating audio samples. */
 	Audio::FPatchOutputStrongPtr PatchOutput;
 
-	/** Callback called when a buffer is received. */
+	/** Callback called when a buffer is received. (Not used, deprecated) */
 	FOnBufferReceived BufferReceivedDelegate;
 };
 
@@ -181,26 +182,32 @@ public:
 
 	FMediaIOAudioCapture() = default;
 
-	UE_DEPRECATED(5.4, "This constructor is deprecated. Use Initialize(InAudioDeviceHandle) instead after constructing it.")
-	FMediaIOAudioCapture(const FAudioDeviceHandle& InAudioDeviceHandle);
+	virtual ~FMediaIOAudioCapture() override;
 
-	virtual ~FMediaIOAudioCapture();
-
-	//~ ISubmixBufferListener interface
+	//~ Begin ISubmixBufferListener
 	virtual void OnNewSubmixBuffer(const USoundSubmix* InOwningSubmix, float* InAudioData, int32 InNumSamples, int32 InNumChannels, const int32 InSampleRate, double InAudioClock) override;
 	virtual const FString& GetListenerName() const override;
-
-	/** Initializes audio capture for the given audio device */
+	//~ End ISubmixBufferListener
+	
+	UE_DEPRECATED(5.6, "Use the version that takes in the AudioCapturedDelegate.")
 	void Initialize(const FAudioDeviceHandle& InAudioDeviceHandle);
+	
+	/**
+	 * Initializes audio capture for the given audio device
+	 * The audio capture delegate can't be changed post Initialize.
+	 */
+	void Initialize(const FAudioDeviceHandle& InAudioDeviceHandle, const FOnAudioCaptured& InAudioCapturedDelegate);
 
 	/** Create an audio output that will receive audio samples. */
 	TSharedPtr<FMediaIOAudioOutput> CreateAudioOutput(int32 InNumOutputChannels, FFrameRate InTargetFrameRate, uint32 InMaxSampleLatency, uint32 InOutputSampleRate);
 
+	UE_DEPRECATED(5.6, "This function is not thread safe. AudioCapturedDelegate must be passed as parameter to Initialize instead.")
 	FOnAudioCaptured& OnAudioCaptured_RenderThread()
 	{
 		return AudioCapturedDelegate;
 	}
 
+	/** Returns the number of channels from the audio mixer device (engine side). */
 	int32 GetNumInputChannels() const
 	{
 		return NumChannels;
@@ -213,6 +220,9 @@ protected:
 	void RegisterAudioDevice(const FAudioDeviceHandle& InAudioDeviceHandle);
 	void UnregisterAudioDevice();
 
+	/** This must be called in Initialize only, and before calling RegisterBufferListener. */
+	void SetAudioCapturedDelegate(const FOnAudioCaptured& InAudioCapturedDelegate);
+
 private:
 	/** Audio device Id this buffer listener is registered to. */
 	Audio::FDeviceId RegisteredDeviceId = INDEX_NONE;
@@ -220,7 +230,13 @@ private:
 	/** Sample rate on the engine side. */ 
 	uint32 SampleRate = 0;
 
-	// Used to make sure we only accumulate audio from the primary submix. 
+	/** Internal implementation of ISubmixBufferListener. */
+	class FSubmixBufferListenerProxy;
+
+	/** We need to proxy the submit buffer listener to be able to unregister it in the destructor. */
+	TSharedPtr<ISubmixBufferListener> SubmixBufferListenerProxy;
+	
+	/** Used to make sure we only accumulate audio from the primary submix. */ 
 	FName PrimarySubmixName;
 
 	/** Number of channels on the engine side. */
@@ -245,9 +261,12 @@ public:
 	FMainMediaIOAudioCapture();
 	virtual ~FMainMediaIOAudioCapture() override;
 
-	/** Initializes audio capture for the main audio device */
+	UE_DEPRECATED(5.6, "Use the version that takes in the AudioCapturedDelegate.")
 	void Initialize();
 
+	/** Initializes audio capture for the main audio device */
+	void Initialize(const FOnAudioCaptured& InAudioCapturedDelegate);
+
 private:
 #if WITH_EDITOR
 	void OnPIEStarted(const bool);
diff --git a/Engine/Plugins/Media/MediaIOFramework/Source/MediaIOCore/Public/MediaIOCoreSubsystem.h b/Engine/Plugins/Media/MediaIOFramework/Source/MediaIOCore/Public/MediaIOCoreSubsystem.h
--- a/Engine/Plugins/Media/MediaIOFramework/Source/MediaIOCore/Public/MediaIOCoreSubsystem.h
+++ b/Engine/Plugins/Media/MediaIOFramework/Source/MediaIOCore/Public/MediaIOCoreSubsystem.h
@@ -14,7 +14,7 @@ UCLASS()
 class MEDIAIOCORE_API UMediaIOCoreSubsystem : public UEngineSubsystem
 {
 public:
-	DECLARE_MULTICAST_DELEGATE_ThreeParams(FOnBufferReceived, Audio::FDeviceId /* DeviceId */, float* /* Data */, int32 /* NumSamples */)
+	DECLARE_TS_MULTICAST_DELEGATE_ThreeParams(FOnBufferReceived, Audio::FDeviceId /* DeviceId */, float* /* Data */, int32 /* NumSamples */)
 	
 	struct FCreateAudioOutputArgs
 	{
@@ -62,5 +62,6 @@ private:
 
 	FDelegateHandle DeviceDestroyedHandle;
 
+	/** Delegate called from the audio thread when a buffer is received. Must be thread safe. */
 	FOnBufferReceived BufferReceivedDelegate;
 };
