Media Capture - Allow the capture implementation to decide whether to discard a frame based on timecode and frame number

diff --git a/Engine/Plugins/Media/MediaIOFramework/Source/MediaIOCore/Private/MediaCapture.cpp b/Engine/Plugins/Media/MediaIOFramework/Source/MediaIOCore/Private/MediaCapture.cpp
--- a/Engine/Plugins/Media/MediaIOFramework/Source/MediaIOCore/Private/MediaCapture.cpp
+++ b/Engine/Plugins/Media/MediaIOFramework/Source/MediaIOCore/Private/MediaCapture.cpp
@@ -1005,7 +1005,7 @@ void UMediaCapture::WaitForPendingTasks()
 void UMediaCapture::ProcessCapture_GameThread()
 {
 	using namespace UE::MediaCaptureData;
-	
+
 	// Acquire a frame
 	TSharedPtr<FCaptureFrame> CapturingFrame;
 	if(GetState() != EMediaCaptureState::StopRequested)
@@ -1410,7 +1410,7 @@ bool UMediaCapture::ProcessCapture_RenderThread(const TSharedPtr<UE::MediaCaptur
 				Args.MediaCapture->CaptureRenderPipeline->InitializeResources_RenderThread(CapturingFrame, Args.GraphBuilder, Args.RDGResourceToCapture);
 			}
 
-			const bool bIsReadyToStartOutput = Args.MediaCapture->IsReadyToOutput();
+			bool bIsReadyToStartOutput = Args.MediaCapture->IsReadyToOutput();
 
 			if (!CapturingFrame->bPrerollFrame && !bIsReadyToStartOutput && !bLoggedOutputNotReady)
 			{
@@ -1420,35 +1420,48 @@ bool UMediaCapture::ProcessCapture_RenderThread(const TSharedPtr<UE::MediaCaptur
 
 			if (!CapturingFrame->bPrerollFrame && bIsReadyToStartOutput)
 			{
-				// Call the capture frame algo based on the specific type of resource we are using
-				if (DesiredOutputResourceType == EMediaCaptureResourceType::Texture)
+				UE_LOG(LogMediaIOCore, Verbose, TEXT("Capturing %s frame (%d)"), FApp::GetTimecode().Frames % 2 == 0 ? TEXT("Even") : TEXT("Odd"), FApp::GetTimecode().Frames);
+
+				if (Args.MediaCapture->ShouldCaptureThisFrame(CapturingFrame->CaptureBaseData))
 				{
-					if (ensure(CapturingFrame->IsTextureResource()))
+					// Call the capture frame algo based on the specific type of resource we are using
+					if (DesiredOutputResourceType == EMediaCaptureResourceType::Texture)
 					{
-						Args.MediaCapture->BeforeFrameCaptured_RenderingThread(CapturingFrame->CaptureBaseData, CapturingFrame->UserData, CapturingFrame->GetTextureResource());
+						if (ensure(CapturingFrame->IsTextureResource()))
+						{
+							Args.MediaCapture->BeforeFrameCaptured_RenderingThread(CapturingFrame->CaptureBaseData, CapturingFrame->UserData, CapturingFrame->GetTextureResource());
+						}
+						else
+						{
+							bBeforeCaptureSuccess = false;
+							UE_LOG(LogMediaIOCore, Error, TEXT("The capture will stop for '%s'. Capture frame was expected to use Texture resource but wasn't."), *Args.MediaCapture->MediaOutputName);
+						}
 					}
 					else
 					{
-						bBeforeCaptureSuccess = false;
-						UE_LOG(LogMediaIOCore, Error, TEXT("The capture will stop for '%s'. Capture frame was expected to use Texture resource but wasn't."), *Args.MediaCapture->MediaOutputName);
-					}
-				}
-				else
-				{
-					if (ensure(CapturingFrame->IsBufferResource()))
-					{
-						Args.MediaCapture->BeforeFrameCaptured_RenderingThread(CapturingFrame->CaptureBaseData, CapturingFrame->UserData, CapturingFrame->GetBufferResource());
+						if (ensure(CapturingFrame->IsBufferResource()))
+						{
+							Args.MediaCapture->BeforeFrameCaptured_RenderingThread(CapturingFrame->CaptureBaseData, CapturingFrame->UserData, CapturingFrame->GetBufferResource());
+						}
+						else
+						{
+							bBeforeCaptureSuccess = false;
+							UE_LOG(LogMediaIOCore, Error, TEXT("The capture will stop for '%s'. Capture frame was expected to use Buffer resource but wasn't."), *Args.MediaCapture->MediaOutputName);
+						}
 					}
-					else
+
+					if (bBeforeCaptureSuccess)
 					{
-						bBeforeCaptureSuccess = false;
-						UE_LOG(LogMediaIOCore, Error, TEXT("The capture will stop for '%s'. Capture frame was expected to use Buffer resource but wasn't."), *Args.MediaCapture->MediaOutputName);
+						bHasCaptureSuceeded = UE::MediaCaptureData::FMediaCaptureHelper::CaptureFrame(Args, CapturingFrame);
 					}
 				}
-
-				if (bBeforeCaptureSuccess)
+				else
 				{
-					bHasCaptureSuceeded = UE::MediaCaptureData::FMediaCaptureHelper::CaptureFrame(Args, CapturingFrame);
+					// Return the frame to the available pool if we're not gonna capture it.
+					bHasCaptureSuceeded = true;
+					FrameManager->MarkAvailable(*CapturingFrame);
+
+					UE_LOG(LogMediaIOCore, Verbose, TEXT("=============================================================== Skipping frame %d"), FApp::GetTimecode().Frames);
 				}
 			}
 			else
diff --git a/Engine/Plugins/Media/MediaIOFramework/Source/MediaIOCore/Public/MediaCapture.h b/Engine/Plugins/Media/MediaIOFramework/Source/MediaIOCore/Public/MediaCapture.h
--- a/Engine/Plugins/Media/MediaIOFramework/Source/MediaIOCore/Public/MediaCapture.h
+++ b/Engine/Plugins/Media/MediaIOFramework/Source/MediaIOCore/Public/MediaCapture.h
@@ -497,6 +497,9 @@ protected:
 	/** Should we call OnFrameCaptured_RenderingThread() with a RHI resource -or- readback the memory to CPU ram and call OnFrameCaptured_RenderingThread(). */
 	virtual bool ShouldCaptureRHIResource() const { return false; }
 
+	/** Allows capture implementaion to decide whether to discard a frame based on its base data. */ 
+	virtual bool ShouldCaptureThisFrame(const FCaptureBaseData& InBaseData) const { return true; }
+
 	/** Called at the beginning of the Capture_RenderThread call if output resource type is Texture */
 	virtual void BeforeFrameCaptured_RenderingThread(const FCaptureBaseData& InBaseData, TSharedPtr<FMediaCaptureUserData, ESPMode::ThreadSafe> InUserData, FTextureRHIRef InTexture) {}
 
