Re-enable MediaIOCore from MediaIOFramework module on Windows Arm64
- the module defines MediaIOShaders for shader source directory mapping

#rnx

diff --git a/Engine/Plugins/Media/MediaIOFramework/MediaIOFramework.uplugin b/Engine/Plugins/Media/MediaIOFramework/MediaIOFramework.uplugin
--- a/Engine/Plugins/Media/MediaIOFramework/MediaIOFramework.uplugin
+++ b/Engine/Plugins/Media/MediaIOFramework/MediaIOFramework.uplugin
@@ -17,8 +17,7 @@
 		{
 			"Name": "MediaIOCore",
 			"Type": "Runtime",
-			"LoadingPhase": "PostConfigInit",
-			"PlatformArchitectureDenyList": [ "Win64:arm64" ]
+			"LoadingPhase": "PostConfigInit"
 		},
 		{
 			"Name": "MediaIOEditor",
