MediaProfile: Fixed issue where function was getting inlined when it wasn't supposed to

diff --git a/Engine/Plugins/Media/MediaIOFramework/Source/MediaIOCore/Private/MediaCapture.cpp b/Engine/Plugins/Media/MediaIOFramework/Source/MediaIOCore/Private/MediaCapture.cpp
--- a/Engine/Plugins/Media/MediaIOFramework/Source/MediaIOCore/Private/MediaCapture.cpp
+++ b/Engine/Plugins/Media/MediaIOFramework/Source/MediaIOCore/Private/MediaCapture.cpp
@@ -128,6 +128,10 @@ void FMediaCaptureOptions::PostSerialize(const FArchive& Ar)
 #endif
 }
 
+FString FMediaCaptureOptions::GetColorConversionSettingsString() const
+{
+	return ColorConversionSettings.ToString();
+}
 
 /* UMediaCapture
 *****************************************************************************/
diff --git a/Engine/Plugins/Media/MediaIOFramework/Source/MediaIOCore/Public/MediaCapture.h b/Engine/Plugins/Media/MediaIOFramework/Source/MediaIOCore/Public/MediaCapture.h
--- a/Engine/Plugins/Media/MediaIOFramework/Source/MediaIOCore/Public/MediaCapture.h
+++ b/Engine/Plugins/Media/MediaIOFramework/Source/MediaIOCore/Public/MediaCapture.h
@@ -176,7 +176,7 @@ public:
 	UE_API void PostSerialize(const FArchive& Ar);
 
 	/** Gets the string version of ColorConversionSettings */
-	UE_API FString GetColorConversionSettingsString() const { return ColorConversionSettings.ToString(); }
+	UE_API FString GetColorConversionSettingsString() const;
 
 	/** If source texture's size change, auto restart capture to follow source resolution if implementation supports it. */
 	UPROPERTY(BlueprintReadWrite, EditAnywhere, Category = "MediaCapture")
