MediaIO: Added RenderTargetMediaOutput (and Capture)

Useful to link media output sources with texture inputs.

In particular, useful to run Vcam with nDisplay, where nDisplay camera render can be output to a render target and picked up by [REDACTED_PERSON] output provider which can make it available to other providers such as pixel streaming or media output provider.

Fully private implementation for now, no API additions.

#rnx

diff --git a/Engine/Plugins/Media/MediaIOFramework/Source/MediaIOCore/Private/RenderTargetMediaCapture.cpp b/Engine/Plugins/Media/MediaIOFramework/Source/MediaIOCore/Private/RenderTargetMediaCapture.cpp
new file mode 100644
--- /dev/null
+++ b/Engine/Plugins/Media/MediaIOFramework/Source/MediaIOCore/Private/RenderTargetMediaCapture.cpp
@@ -0,0 +1,113 @@
+// Copyright Epic Games, Inc. All Rights Reserved.
+
+#include "RenderTargetMediaCapture.h"
+
+#include "Engine/Engine.h"
+#include "Engine/TextureRenderTarget2D.h"
+#include "MediaIOCoreModule.h"
+#include "RenderTargetMediaOutput.h"
+#include "Slate/SceneViewport.h"
+#include "TextureResource.h"
+
+
+void URenderTargetMediaCapture::OnRHIResourceCaptured_RenderingThread(
+	FRHICommandListImmediate& RHICmdList, 
+	const FCaptureBaseData& InBaseData, 
+	TSharedPtr<FMediaCaptureUserData, 
+	ESPMode::ThreadSafe> InUserData, 
+	FTextureRHIRef InTexture
+)
+{
+	FScopeLock ScopeLock(&RenderTargetCriticalSection);
+
+	if (!RenderTarget || !RenderTarget->GetRenderTargetResource())
+	{
+		return;
+	}
+	
+	const FTextureRHIRef RenderTargetRHIRef = RenderTarget->GetRenderTargetResource()->GetRenderTargetTexture();
+
+	if (!RenderTargetRHIRef.IsValid())
+	{
+		return;
+	}
+
+	// copy the captured texture into our render target
+
+	RHICmdList.CopyTexture(
+		InTexture,
+		RenderTargetRHIRef,
+		FRHICopyTextureInfo()
+	);
+
+	RenderTarget->GetRenderTargetResource()->FlushDeferredResourceUpdate(RHICmdList);
+}
+
+bool URenderTargetMediaCapture::InitializeCapture()
+{
+	return true;
+}
+
+bool URenderTargetMediaCapture::PostInitializeCaptureViewport(TSharedPtr<FSceneViewport>& InSceneViewport)
+{
+	const FTextureRHIRef& BackBuffer = InSceneViewport->GetRenderTargetTexture();
+
+	if (!BackBuffer.IsValid())
+	{
+		UE_LOG(LogMediaIOCore, Error, TEXT("RenderTargetMediaCapture: Scene viewport had no backbuffer."));
+		return false;
+	}
+
+	const FRHITextureDesc& Desc = BackBuffer->GetDesc();
+	return StartNewCapture(Desc.Extent, Desc.Format);
+}
+
+bool URenderTargetMediaCapture::PostInitializeCaptureRenderTarget(UTextureRenderTarget2D* InRenderTarget)
+{
+	return StartNewCapture(FIntPoint(InRenderTarget->SizeX, InRenderTarget->SizeY), InRenderTarget->GetFormat());
+}
+
+bool URenderTargetMediaCapture::PostInitializeCaptureRHIResource(const FRHICaptureResourceDescription& InResourceDescription)
+{
+	return StartNewCapture(InResourceDescription.ResourceSize, InResourceDescription.PixelFormat);
+}
+
+void URenderTargetMediaCapture::StopCaptureImpl(bool /*bAllowPendingFrameToBeProcess*/)
+{
+	TRACE_BOOKMARK(TEXT("RenderTargetMediaCapture::StopCapture"));
+}
+
+bool URenderTargetMediaCapture::StartNewCapture(const FIntPoint& InSourceTargetSize, EPixelFormat InSourceTargetFormat)
+{
+	TRACE_BOOKMARK(TEXT("RenderTargetMediaCapture::StartNewCapture"));
+	{
+		RenderTarget = nullptr;
+		
+		const URenderTargetMediaOutput* RenderTargetMediaOutput = CastChecked<URenderTargetMediaOutput>(MediaOutput);
+
+		if (RenderTargetMediaOutput)
+		{
+			FScopeLock ScopeLock(&RenderTargetCriticalSection);
+
+			RenderTarget = RenderTargetMediaOutput->RenderTarget.LoadSynchronous();
+
+			if (RenderTarget)
+			{
+				SetState(EMediaCaptureState::Capturing);
+				return true;
+			}
+			else
+			{
+				UE_LOG(LogMediaIOCore, Error, TEXT("RenderTargetMediaCapture: Missing render target object: '%s'."), *RenderTargetMediaOutput->RenderTarget.ToString());
+			}
+		}
+		else
+		{
+			UE_LOG(LogMediaIOCore, Error, TEXT("RenderTargetMediaCapture: Media Output not found or of unexpected type."));
+		}
+	}
+
+	SetState(EMediaCaptureState::Error);
+
+	return false;
+}
diff --git a/Engine/Plugins/Media/MediaIOFramework/Source/MediaIOCore/Private/RenderTargetMediaCapture.h b/Engine/Plugins/Media/MediaIOFramework/Source/MediaIOCore/Private/RenderTargetMediaCapture.h
new file mode 100644
--- /dev/null
+++ b/Engine/Plugins/Media/MediaIOFramework/Source/MediaIOCore/Private/RenderTargetMediaCapture.h
@@ -0,0 +1,44 @@
+// Copyright Epic Games, Inc. All Rights Reserved.
+
+#pragma once
+
+#include "HAL/CriticalSection.h"
+#include "MediaCapture.h"
+
+#include "RenderTargetMediaCapture.generated.h"
+
+
+UCLASS()
+class URenderTargetMediaCapture : public UMediaCapture
+{
+	GENERATED_BODY()
+
+protected:
+
+	virtual bool ShouldCaptureRHIResource() const override
+	{ 
+		return true; 
+	}
+
+	virtual void OnRHIResourceCaptured_RenderingThread(
+		FRHICommandListImmediate& RHICmdList, 
+		const FCaptureBaseData& InBaseData, 
+		TSharedPtr<FMediaCaptureUserData, 
+		ESPMode::ThreadSafe> InUserData, 
+		FTextureRHIRef InTexture
+	) override;
+
+	virtual bool InitializeCapture() override;
+	virtual bool PostInitializeCaptureViewport(TSharedPtr<FSceneViewport>& InSceneViewport) override;
+	virtual bool PostInitializeCaptureRenderTarget(UTextureRenderTarget2D* InRenderTarget) override;
+	virtual bool PostInitializeCaptureRHIResource(const FRHICaptureResourceDescription& InResourceDescription) override;
+	virtual void StopCaptureImpl(bool bAllowPendingFrameToBeProcess) override;
+
+	UPROPERTY(Transient)
+	TObjectPtr<UTextureRenderTarget2D> RenderTarget;
+
+private:
+	bool StartNewCapture(const FIntPoint& InSourceTargetSize, EPixelFormat InSourceTargetFormat);
+	
+	FCriticalSection RenderTargetCriticalSection;
+};
diff --git a/Engine/Plugins/Media/MediaIOFramework/Source/MediaIOCore/Private/RenderTargetMediaOutput.cpp b/Engine/Plugins/Media/MediaIOFramework/Source/MediaIOCore/Private/RenderTargetMediaOutput.cpp
new file mode 100644
--- /dev/null
+++ b/Engine/Plugins/Media/MediaIOFramework/Source/MediaIOCore/Private/RenderTargetMediaOutput.cpp
@@ -0,0 +1,69 @@
+// Copyright Epic Games, Inc. All Rights Reserved.
+
+#include "RenderTargetMediaOutput.h"
+
+#include "Engine/TextureRenderTarget2D.h"
+#include "MediaIOCoreModule.h"
+#include "RenderTargetMediaCapture.h"
+
+bool URenderTargetMediaOutput::Validate(FString& OutFailureReason) const
+{
+	if (!Super::Validate(OutFailureReason))
+	{
+		return false;
+	}
+
+	const UTextureRenderTarget2D* LoadedRenderTarget = RenderTarget.LoadSynchronous();
+
+	if (!LoadedRenderTarget)
+	{
+		OutFailureReason = FString::Printf(TEXT("Missing render target object: '%s'."), *RenderTarget.ToString());
+		return false;
+	}
+
+	return true;
+}
+
+FIntPoint URenderTargetMediaOutput::GetRequestedSize() const
+{
+	if (const UTextureRenderTarget2D* LoadedRenderTarget = RenderTarget.LoadSynchronous())
+	{
+		return FIntPoint(LoadedRenderTarget->SizeX, LoadedRenderTarget->SizeY);
+	}
+
+	return RequestCaptureSourceSize;
+}
+
+
+EPixelFormat URenderTargetMediaOutput::GetRequestedPixelFormat() const
+{
+	if (const UTextureRenderTarget2D* LoadedRenderTarget = RenderTarget.LoadSynchronous())
+	{
+		return LoadedRenderTarget->GetFormat();
+	}
+
+	return EPixelFormat::PF_Unknown;
+}
+
+EMediaCaptureConversionOperation URenderTargetMediaOutput::GetConversionOperation(EMediaCaptureSourceType /*InSourceType*/) const
+{
+	if (bInvertAlpha)
+	{
+		return EMediaCaptureConversionOperation::INVERT_ALPHA;
+	}
+
+	return EMediaCaptureConversionOperation::NONE;
+}
+
+UMediaCapture* URenderTargetMediaOutput::CreateMediaCaptureImpl()
+{
+	UMediaCapture* Result = NewObject<URenderTargetMediaCapture>();
+
+	if (Result)
+	{
+		UE_LOG(LogMediaIOCore, Log, TEXT("Created Render Target Media Capture"));
+		Result->SetMediaOutput(this);
+	}
+
+	return Result;
+}
diff --git a/Engine/Plugins/Media/MediaIOFramework/Source/MediaIOCore/Private/RenderTargetMediaOutput.h b/Engine/Plugins/Media/MediaIOFramework/Source/MediaIOCore/Private/RenderTargetMediaOutput.h
new file mode 100644
--- /dev/null
+++ b/Engine/Plugins/Media/MediaIOFramework/Source/MediaIOCore/Private/RenderTargetMediaOutput.h
@@ -0,0 +1,44 @@
+// Copyright Epic Games, Inc. All Rights Reserved.
+
+#pragma once
+
+#include "MediaOutput.h"
+#include "UObject/SoftObjectPtr.h"
+#include "RenderTargetMediaOutput.generated.h"
+
+class UTextureRenderTarget2D;
+
+/**
+ * Captures Media to a Render Target. This can be useful to make a texture available
+ * from a source that has a Media output interface.
+ */
+UCLASS(BlueprintType)
+class URenderTargetMediaOutput : public UMediaOutput
+{
+	GENERATED_BODY()
+
+public:
+
+	/** Invert the alpha channel */
+	UPROPERTY(BlueprintReadWrite, EditAnywhere, Category="Media")
+	bool bInvertAlpha = false;
+
+	/** Specify the render target to be capturing to. */
+	UPROPERTY(BlueprintReadWrite, EditAnywhere, Category="Media")
+	TSoftObjectPtr<UTextureRenderTarget2D> RenderTarget;
+	
+	//~ Begin UMediaOutput interface
+
+public:
+
+	virtual bool Validate(FString& FailureReason) const override;
+	virtual FIntPoint GetRequestedSize() const override;
+	virtual EPixelFormat GetRequestedPixelFormat() const override;
+	virtual EMediaCaptureConversionOperation GetConversionOperation(EMediaCaptureSourceType InSourceType) const override;
+
+protected:
+
+	virtual UMediaCapture* CreateMediaCaptureImpl() override;
+
+	//~ End UMediaOutput interface
+};
