MediaCapture: Ensure that the first immediate capture has frame info available.

The game frame data structure is currently populated and added to the queue in a FCoreDelegates::OnBeginFrame callback, and when the capture is started that delegate has typically already fired. So the first render thread capture will typically be missing the info for that frame and fail the capture.

This CL manually adds that first game frame data into the queue.

#rnx

diff --git a/Engine/Plugins/Media/MediaIOFramework/Source/MediaIOCore/Private/MediaCapture.cpp b/Engine/Plugins/Media/MediaIOFramework/Source/MediaIOCore/Private/MediaCapture.cpp
--- a/Engine/Plugins/Media/MediaIOFramework/Source/MediaIOCore/Private/MediaCapture.cpp
+++ b/Engine/Plugins/Media/MediaIOFramework/Source/MediaIOCore/Private/MediaCapture.cpp
@@ -326,6 +326,9 @@ bool UMediaCapture::StartSourceCapture(TSharedPtr<UE::MediaCapture::Private::FCa
 		{
 			// Immediate capture requires us to prepare frame info in OnBeginFrame
 			FCoreDelegates::OnBeginFrame.AddUObject(this, &UMediaCapture::OnBeginFrame_GameThread);
+
+			// Add frame info for current frame in case OnBeginFrame delegate has already fired
+			AddFrameInfoForCaptureImmediate(GFrameCounter);
 		}
 		else
 		{
@@ -1293,6 +1296,32 @@ bool UMediaCapture::CaptureRenderTargetImpl(UTextureRenderTarget2D * InRenderTar
 	return false;
 }
 
+void UMediaCapture::AddFrameInfoForCaptureImmediate(const uint64 FrameNumber)
+{
+	// Prepare capture data
+
+	FQueuedCaptureData CaptureData;
+	CaptureData.BaseData.SourceFrameTimecode = FApp::GetTimecode();
+	CaptureData.BaseData.SourceFrameTimecodeFramerate = FApp::GetTimecodeFrameRate();
+	CaptureData.BaseData.SourceFrameNumberRenderThread = FrameNumber;
+	CaptureData.UserData = GetCaptureFrameUserData_GameThread();
+
+	// Insert it into the queue
+	{
+		FScopeLock Lock(&CaptureDataQueueCriticalSection);
+		
+		// Don't add if we already cached this frame number or newer.
+		// Note: Only checks last inserted since frames are expected to increase monotonically and inserted at position 0.
+		if (CaptureDataQueue.Num() > 0 && CaptureDataQueue[0].BaseData.SourceFrameNumberRenderThread >= FrameNumber)
+		{
+			return;
+		}
+
+		TRACE_BOOKMARK(TEXT("MediaCapture AddFrameInfoForCaptureImmediate for frame %d"), FrameNumber);
+		CaptureDataQueue.Insert(MoveTemp(CaptureData), 0);
+	}
+}
+
 void UMediaCapture::OnBeginFrame_GameThread()
 {
 	TRACE_CPUPROFILER_EVENT_SCOPE(MediaCapture::BeginFrame);
@@ -1310,18 +1339,7 @@ void UMediaCapture::OnBeginFrame_GameThread()
 
 	if (CaptureSource->GetCaptureType() == UE::MediaCapture::Private::ECaptureType::Immediate && (GetState() == EMediaCaptureState::Capturing))
 	{
-		// Queue capture data to be consumed when capture requests are done on render thread
-		FQueuedCaptureData CaptureData;
-		CaptureData.BaseData.SourceFrameTimecode = FApp::GetTimecode();
-		CaptureData.BaseData.SourceFrameTimecodeFramerate = FApp::GetTimecodeFrameRate();
-		CaptureData.BaseData.SourceFrameNumberRenderThread = GFrameCounter;
-		CaptureData.UserData = GetCaptureFrameUserData_GameThread();
-
-		FScopeLock Lock(&CaptureDataQueueCriticalSection);
-
-		TRACE_BOOKMARK(TEXT("MediaCapture Push frame %d"), GFrameCounter);
-		CaptureDataQueue.Insert(MoveTemp(CaptureData), 0);
-
+		AddFrameInfoForCaptureImmediate(GFrameCounter);
 	}
 }
 
diff --git a/Engine/Plugins/Media/MediaIOFramework/Source/MediaIOCore/Public/MediaCapture.h b/Engine/Plugins/Media/MediaIOFramework/Source/MediaIOCore/Public/MediaCapture.h
--- a/Engine/Plugins/Media/MediaIOFramework/Source/MediaIOCore/Public/MediaCapture.h
+++ b/Engine/Plugins/Media/MediaIOFramework/Source/MediaIOCore/Public/MediaCapture.h
@@ -627,7 +627,6 @@ protected:
 private:
 	UE_API void InitializeOutputResources(int32 InNumberOfBuffers);
 	
-	/** Used by RHI capture to buffer captured frame data dependant on game thread such as timecode */
 	UE_API void OnBeginFrame_GameThread();
 	UE_API void OnEndFrame_GameThread();
 	UE_API void CacheMediaOutput(EMediaCaptureSourceType InSourceType);
@@ -673,6 +672,9 @@ private:
 	 */
 	UE_API bool UseAnyThreadCapture() const;
 
+	/** Caches game thread dependent frame data for use by Capture Immediate. */
+	void AddFrameInfoForCaptureImmediate(const uint64 FrameNumber);
+
 	/** Returns bytes per pixel based on pixel format */
 	static UE_API int32 GetBytesPerPixel(EPixelFormat InPixelFormat);
 
