Aja - Add version of StopCapture method that returns a Future to allow being notified when the underlying implementation is done cleaning up its resources.

This is required when calling StartCapture while it's already running, in order to allow the Aja SDK to clean reset the state of the channel before it's used again.

Without this change, trying to start a capture on a MediaCapture that had entered the Error state would fail due to the underlying channel still being in use. This issue was particularly visible in workflows that automatically restart captures on error, such as in DisplayClusterMediaCaptureBase.

diff --git a/Engine/Plugins/Media/AjaMedia/Source/Aja/AjaCore.Build.cs b/Engine/Plugins/Media/AjaMedia/Source/Aja/AjaCore.Build.cs
--- a/Engine/Plugins/Media/AjaMedia/Source/Aja/AjaCore.Build.cs
+++ b/Engine/Plugins/Media/AjaMedia/Source/Aja/AjaCore.Build.cs
@@ -14,7 +14,8 @@ public class AjaCore : ModuleRules
 		PrivateDependencyModuleNames.AddRange(
 			new string[]
 			{
-				"GPUTextureTransfer"
+				"GPUTextureTransfer",
+				"RenderCore"
 			});
 			
 		
diff --git a/Engine/Plugins/Media/AjaMedia/Source/Aja/Private/AJA.cpp b/Engine/Plugins/Media/AjaMedia/Source/Aja/Private/AJA.cpp
--- a/Engine/Plugins/Media/AjaMedia/Source/Aja/Private/AJA.cpp
+++ b/Engine/Plugins/Media/AjaMedia/Source/Aja/Private/AJA.cpp
@@ -414,9 +414,9 @@ namespace AJA
 		return Channel->Initialize(InDevice, Options);
 	}
 
-	void AJAOutputChannel::Uninitialize()
+	void AJAOutputChannel::Uninitialize(TPromise<void> Promise)
 	{
-		Channel->Uninitialize();
+		Channel->Uninitialize(MoveTemp(Promise));
 	}
 
 	AJAOutputChannel::~AJAOutputChannel()
diff --git a/Engine/Plugins/Media/AjaMedia/Source/Aja/Private/ChannelBase.cpp b/Engine/Plugins/Media/AjaMedia/Source/Aja/Private/ChannelBase.cpp
--- a/Engine/Plugins/Media/AjaMedia/Source/Aja/Private/ChannelBase.cpp
+++ b/Engine/Plugins/Media/AjaMedia/Source/Aja/Private/ChannelBase.cpp
@@ -28,9 +28,10 @@ namespace AJA
 
 		/* IOChannelUninitialize_DeviceCommand implementation
 		*****************************************************************************/
-		IOChannelUninitialize_DeviceCommand::IOChannelUninitialize_DeviceCommand(const std::shared_ptr<DeviceConnection>& InConnection, const std::weak_ptr<ChannelThreadBase>& InChannel)
+		IOChannelUninitialize_DeviceCommand::IOChannelUninitialize_DeviceCommand(const std::shared_ptr<DeviceConnection>& InConnection, const std::weak_ptr<ChannelThreadBase>& InChannel, TPromise<void> InCompletionPromise)
 			: DeviceCommand(InConnection)
 			, ChannelThread(InChannel)
+			, CompletionPromise(MoveTemp(InCompletionPromise))
 		{}
 
 		void IOChannelUninitialize_DeviceCommand::Execute(DeviceConnection::CommandList& InCommandList)
@@ -41,6 +42,7 @@ namespace AJA
 			ChannelThread->FrameThread.Stop();
 			ChannelThread->DeviceThread_Destroy(InCommandList);
 			ChannelThread.reset();
+			CompletionPromise.SetValue();
 		}
 
 		/* ChannelThreadBase implementation
diff --git a/Engine/Plugins/Media/AjaMedia/Source/Aja/Private/ChannelBase.h b/Engine/Plugins/Media/AjaMedia/Source/Aja/Private/ChannelBase.h
--- a/Engine/Plugins/Media/AjaMedia/Source/Aja/Private/ChannelBase.h
+++ b/Engine/Plugins/Media/AjaMedia/Source/Aja/Private/ChannelBase.h
@@ -31,11 +31,15 @@ namespace AJA
 		*****************************************************************************/
 		struct IOChannelUninitialize_DeviceCommand : DeviceCommand
 		{
-			IOChannelUninitialize_DeviceCommand(const std::shared_ptr<DeviceConnection>& InConnection, const std::weak_ptr<ChannelThreadBase>& InChannel);
+			IOChannelUninitialize_DeviceCommand(const std::shared_ptr<DeviceConnection>& InConnection, const std::weak_ptr<ChannelThreadBase>& InChannel, TPromise<void> InCompletionPromise = {});
 			
 			virtual void Execute(DeviceConnection::CommandList& InCommandList) override;
 
 			std::shared_ptr<ChannelThreadBase> ChannelThread;
+
+		private:
+			/** Completion promise fired when the uninitialized command is done. */
+			TPromise<void> CompletionPromise;
 		};
 
 		/* ChannelThreadBase OutputChannelThread
diff --git a/Engine/Plugins/Media/AjaMedia/Source/Aja/Private/OutputChannel.cpp b/Engine/Plugins/Media/AjaMedia/Source/Aja/Private/OutputChannel.cpp
--- a/Engine/Plugins/Media/AjaMedia/Source/Aja/Private/OutputChannel.cpp
+++ b/Engine/Plugins/Media/AjaMedia/Source/Aja/Private/OutputChannel.cpp
@@ -2,8 +2,8 @@
 
 #include "OutputChannel.h"
 
-
 #include "GPUTextureTransfer.h"
+#include "RenderingThread.h"
 #include "SSE2MemCpy.h"
 #include <stdlib.h>
 
@@ -19,7 +19,8 @@ namespace AJA
 		*****************************************************************************/
 		bool OutputChannel::Initialize(const AJADeviceOptions& InDeviceOption, const AJAInputOutputChannelOptions& InOptions)
 		{
-			Uninitialize();
+			TPromise<void> CompletionPromise;
+			Uninitialize(MoveTemp(CompletionPromise));
 
 			ChannelThread.reset(new OutputChannelThread(InDeviceOption, InOptions));
 
@@ -35,7 +36,7 @@ namespace AJA
 			return bResult;
 		}
 
-		void OutputChannel::Uninitialize()
+		void OutputChannel::Uninitialize(TPromise<void> CompletionPromise)
 		{
 			if (ChannelThread)
 			{
@@ -56,9 +57,18 @@ namespace AJA
 				if (Device)
 				{
 					std::shared_ptr<IOChannelUninitialize_DeviceCommand> UninitializeCommand;
-					UninitializeCommand.reset(new IOChannelUninitialize_DeviceCommand(Device, ChannelThread));
+					UninitializeCommand.reset(new IOChannelUninitialize_DeviceCommand(Device, ChannelThread, MoveTemp(CompletionPromise)));
 					Device->AddCommand(std::static_pointer_cast<DeviceCommand>(UninitializeCommand));
 				}
+				else
+				{
+					// No need to wait in this case.
+					CompletionPromise.SetValue();
+				}
+			}
+			else
+			{
+				CompletionPromise.SetValue();
 			}
 
 			InitializeCommand.reset();
@@ -220,12 +230,25 @@ namespace AJA
 		void OutputChannelThread::Uninitialize()
 		{
 			ChannelThreadBase::Uninitialize();
+
 			if (TextureTransfer && GetOptions().bUseGPUDMA)
 			{
-				for (auto It = AllFrames.begin(); It!= AllFrames.end(); ++It)
+				std::vector<uint8_t*> BuffersToClear;
+				BuffersToClear.reserve(AllFrames.size());
+				std::transform(AllFrames.cbegin(), AllFrames.cend(), std::back_inserter(BuffersToClear), [](Frame* InFrame) { return InFrame->VideoBuffer; });
+
+				// Make sure to unregister those buffers on the render thread to avoid a crash in the GPUDirect SDK.
+				ENQUEUE_RENDER_COMMAND(ClearGPUDirectBuffers)(
+				[BuffersToClear = std::move(BuffersToClear), TextureTransferPtr = TextureTransfer](FRHICommandListImmediate& RHICmdList)
 				{
-					TextureTransfer->UnregisterBuffer((*It)->VideoBuffer);
-				}
+					if (TextureTransferPtr)
+					{
+						for (uint8* BufferAddress : BuffersToClear)
+						{
+							TextureTransferPtr->UnregisterBuffer(BufferAddress);
+						}
+					}
+				});
 			}
 		}
 
@@ -962,7 +985,7 @@ namespace AJA
 							Transfer.SetAudioBuffer(reinterpret_cast<ULWord*>(AvailableReadingFrame->AudioBuffer), AvailableReadingFrame->CopiedAudioBufferSize);
 						}
 
-						if (UseVideo() && AvailableReadingFrame->CopiedVideoBufferSize && TextureTransfer && GetOptions().bUseGPUDMA)
+						if (UseVideo() && AvailableReadingFrame->CopiedVideoBufferSize && TextureTransfer && GetOptions().bUseGPUDMA && !bStopRequested)
 						{
 							const UE::GPUTextureTransfer::ETransferDirection Direction = UE::GPUTextureTransfer::ETransferDirection::GPU_TO_CPU;
 							TextureTransfer->BeginSync(AvailableReadingFrame->VideoBuffer, Direction);
@@ -983,7 +1006,7 @@ namespace AJA
 						}
 
 						AJA_CHECK(GetDevice().AutoCirculateTransfer(Channel, Transfer));
-						if (UseVideo() && AvailableReadingFrame->CopiedVideoBufferSize && TextureTransfer && GetOptions().bUseGPUDMA)
+						if (UseVideo() && AvailableReadingFrame->CopiedVideoBufferSize && TextureTransfer && GetOptions().bUseGPUDMA && !bStopRequested)
 						{
 							TextureTransfer->EndSync(AvailableReadingFrame->VideoBuffer);
 						}
@@ -1205,7 +1228,7 @@ namespace AJA
 					if (AvailableReadingFrame->CopiedVideoBufferSize > 0)
 					{
 						BurnTimecode(AvailableReadingFrame->Timecode, AvailableReadingFrame->VideoBuffer);
-						if (TextureTransfer && GetOptions().bUseGPUDMA)
+						if (TextureTransfer && GetOptions().bUseGPUDMA && !bStopRequested)
 						{
 							const UE::GPUTextureTransfer::ETransferDirection Direction = UE::GPUTextureTransfer::ETransferDirection::GPU_TO_CPU;
 							TextureTransfer->BeginSync(AvailableReadingFrame->VideoBuffer, Direction);
@@ -1213,7 +1236,7 @@ namespace AJA
 
 						bRunning = bRunning && GetDevice().DMAWriteFrame(BaseFrameIndex + CurrentOutFrame, reinterpret_cast<ULWord*>(AvailableReadingFrame->VideoBuffer), AvailableReadingFrame->CopiedVideoBufferSize, Channel);
 
-						if (TextureTransfer && GetOptions().bUseGPUDMA)
+						if (TextureTransfer && GetOptions().bUseGPUDMA && !bStopRequested)
 						{
 							TextureTransfer->EndSync(AvailableReadingFrame->VideoBuffer);
 						}
diff --git a/Engine/Plugins/Media/AjaMedia/Source/Aja/Private/OutputChannel.h b/Engine/Plugins/Media/AjaMedia/Source/Aja/Private/OutputChannel.h
--- a/Engine/Plugins/Media/AjaMedia/Source/Aja/Private/OutputChannel.h
+++ b/Engine/Plugins/Media/AjaMedia/Source/Aja/Private/OutputChannel.h
@@ -26,7 +26,7 @@ namespace AJA
 			OutputChannel& operator=(const OutputChannel&) = delete;
 
 			bool Initialize(const AJADeviceOptions& InDevice, const AJAInputOutputChannelOptions& Options);
-			void Uninitialize();
+			void Uninitialize(TPromise<void> CompletionPromise);
 
 			bool SetAncillaryFrameData(const AJAOutputFrameBufferData& InFrameData, uint8_t* AncillaryBuffer, uint32_t AncillaryBufferSize);
 			bool SetAudioFrameData(const AJAOutputFrameBufferData& InFrameData, uint8_t* InAudioBuffer, uint32_t InAudioBufferSize);
diff --git a/Engine/Plugins/Media/AjaMedia/Source/Aja/Public/AJALib.h b/Engine/Plugins/Media/AjaMedia/Source/Aja/Public/AJALib.h
--- a/Engine/Plugins/Media/AjaMedia/Source/Aja/Public/AJALib.h
+++ b/Engine/Plugins/Media/AjaMedia/Source/Aja/Public/AJALib.h
@@ -7,6 +7,7 @@
 #include <vector>
 
 #include "AjaCoreModule.h"
+#include "Async/Future.h"
 
 
 class FRHITexture;
@@ -525,7 +526,12 @@ namespace AJA
 
 	public:
 		bool Initialize(const AJADeviceOptions& InDevice, const AJAInputOutputChannelOptions& Options);
-		void Uninitialize();
+	
+		/** 
+		 * Clean up card connections for this channel. (Asynchronous)
+		 * @param Promise Used to inform when the uninitialization is complete.
+		 */
+		void Uninitialize(TPromise<void> Promise = {});
 
 		// Set a new buffer that will be copied to the AJA.
 		bool SetAncillaryFrameData(const AJAOutputFrameBufferData& InFrameData, uint8_t* AncillaryBuffer, uint32_t AncillaryBufferSize);
diff --git a/Engine/Plugins/Media/AjaMedia/Source/AjaMediaOutput/Private/AjaMediaCapture.cpp b/Engine/Plugins/Media/AjaMedia/Source/AjaMediaOutput/Private/AjaMediaCapture.cpp
--- a/Engine/Plugins/Media/AjaMedia/Source/AjaMediaOutput/Private/AjaMediaCapture.cpp
+++ b/Engine/Plugins/Media/AjaMedia/Source/AjaMediaOutput/Private/AjaMediaCapture.cpp
@@ -331,54 +331,47 @@ void UAjaMediaCapture::StopCaptureImpl(bool bAllowPendingFrameToBeProcess)
 {
 	if (!bAllowPendingFrameToBeProcess)
 	{
-		{
-			// Prevent the rendering thread from copying while we are stopping the capture.
-			FScopeLock ScopeLock(&CopyingCriticalSection);
+		StopCapture_Future();
+	}
+}
 
-			ENQUEUE_RENDER_COMMAND(AjaMediaCaptureInitialize)(
-			[this](FRHICommandListImmediate& RHICmdList) mutable
-			{
-				// Unregister texture before closing channel.
-				if (ShouldCaptureRHIResource())
-				{
-					for (FTextureRHIRef& Texture : TexturesToRelease)
-					{
-						TextureTransfer->UnregisterTexture(Texture->GetTexture2D());
-					}
+TFuture<void> UAjaMediaCapture::StopCapture_Future()
+{
+	FScopeLock ScopeLock(&CopyingCriticalSection);
 
-					TexturesToRelease.Reset();
-				}
+	TFuture<void> Future;
+
+	if (OutputChannel)
+	{
+		TRACE_CPUPROFILER_EVENT_SCOPE(UAjaMediaCapture::StopCapture_Future);
+
+		TPromise<void> CompletionPromise;
+		Future = CompletionPromise.GetFuture();
 
-				if (OutputChannel)
+		if (TexturesToRelease.Num())
+		{
+			ENQUEUE_RENDER_COMMAND(AjaMediaCaptureInitialize)(
+			[this, TexturesToUnregister = TexturesToRelease](FRHICommandListImmediate& RHICmdList) mutable
+			{
+				// Unregistering them from the game thread would cause a crash in the GPUDirect SDK.
+				for (FTextureRHIRef& Texture : TexturesToUnregister)
 				{
-					// Close the aja channel in the another thread.
-					OutputChannel->Uninitialize();
-					OutputChannel.Reset();
-					OutputCallback.Reset();
+					TextureTransfer->UnregisterTexture(Texture->GetTexture2D());
 				}
 			});
-			
-		
-			if (WakeUpEvent)
-			{
-				FPlatformProcess::ReturnSynchEventToPool(WakeUpEvent);
-				WakeUpEvent = nullptr;
-			}
-		}
 
-		if (GEngine)
-		{
-			UMediaIOCoreSubsystem* SubSystem = GEngine->GetEngineSubsystem<UMediaIOCoreSubsystem>();
-			if (SubSystem)
-			{
-				SubSystem->OnBufferReceived_AudioThread().RemoveAll(this);
-			}
+			TexturesToRelease.Reset();
 		}
+		
+		OutputChannel->Uninitialize(MoveTemp(CompletionPromise));
+	}
 
-		AudioOutput.Reset();
-
-		RestoreViewportTextureAlpha(GetCapturingSceneViewport());
+	if (!Future.IsValid())
+	{
+		Future = MakeFulfilledPromise<void>().GetFuture();
 	}
+
+	return Future;
 }
 
 bool UAjaMediaCapture::ShouldCaptureRHIResource() const
@@ -435,7 +428,6 @@ void UAjaMediaCapture::OnEnginePreExit()
 {
 	if (OutputChannel)
 	{
-		// Close the aja channel in the another thread.
 		OutputChannel->Uninitialize();
 		OutputChannel.Reset();
 		OutputCallback.Reset();
@@ -889,3 +881,28 @@ void UAjaMediaCapture::OnAudioBufferReceived_AudioThread(Audio::FDeviceId Device
 		}
 	}
 }
+
+void UAjaMediaCapture::PostCaptureCleanupImpl()
+{
+	OutputChannel.Reset();
+	OutputCallback.Reset();
+
+	if (WakeUpEvent)
+	{
+		FPlatformProcess::ReturnSynchEventToPool(WakeUpEvent);
+		WakeUpEvent = nullptr;
+	}
+
+	if (GEngine)
+	{
+		UMediaIOCoreSubsystem* SubSystem = GEngine->GetEngineSubsystem<UMediaIOCoreSubsystem>();
+		if (SubSystem)
+		{
+			SubSystem->OnBufferReceived_AudioThread().RemoveAll(this);
+		}
+	}
+
+	AudioOutput.Reset();
+
+	RestoreViewportTextureAlpha(GetCapturingSceneViewport());
+}
diff --git a/Engine/Plugins/Media/AjaMedia/Source/AjaMediaOutput/Public/AjaMediaCapture.h b/Engine/Plugins/Media/AjaMedia/Source/AjaMediaOutput/Public/AjaMediaCapture.h
--- a/Engine/Plugins/Media/AjaMedia/Source/AjaMediaOutput/Public/AjaMediaCapture.h
+++ b/Engine/Plugins/Media/AjaMedia/Source/AjaMediaOutput/Public/AjaMediaCapture.h
@@ -48,6 +48,9 @@ protected:
 	virtual bool UpdateRenderTargetImpl(UTextureRenderTarget2D* InRenderTarget) override;
 	virtual bool UpdateAudioDeviceImpl(const FAudioDeviceHandle& InAudioDeviceHandle) override;
 	virtual void StopCaptureImpl(bool bAllowPendingFrameToBeProcess) override;
+	virtual void PostCaptureCleanupImpl() override;
+	virtual bool SupportsStopCapture_Future() const override { return true; }
+	virtual TFuture<void> StopCapture_Future() override;
 	virtual bool ShouldCaptureRHIResource() const override;
 	
 	virtual void OnFrameCaptured_RenderingThread(const FCaptureBaseData& InBaseData, TSharedPtr<FMediaCaptureUserData, ESPMode::ThreadSafe> InUserData, void* InBuffer, int32 Width, int32 Height, int32 BytesPerRow) override;
diff --git a/Engine/Plugins/Media/MediaIOFramework/Source/MediaIOCore/Private/MediaCapture.cpp b/Engine/Plugins/Media/MediaIOFramework/Source/MediaIOCore/Private/MediaCapture.cpp
--- a/Engine/Plugins/Media/MediaIOFramework/Source/MediaIOCore/Private/MediaCapture.cpp
+++ b/Engine/Plugins/Media/MediaIOFramework/Source/MediaIOCore/Private/MediaCapture.cpp
@@ -57,6 +57,19 @@ static TAutoConsoleVariable<int32> CVarMediaIOScheduleOnAnyThread(
 	TEXT("Whether to wait for resource readback in a separate thread. (Experimental)"),
 	ECVF_RenderThreadSafe);
 
+static bool bSimulateError = false;
+
+FAutoConsoleCommand FMediaIOSetErrorState
+(
+	TEXT("MediaIO.SimulateError"),
+	TEXT("Sets all media captures to the error state."),
+	FConsoleCommandDelegate::CreateLambda([]()
+		{
+			bSimulateError = true;
+		})
+);
+
+
 DECLARE_GPU_STAT(MediaCapture_ProcessCapture);
 DECLARE_CYCLE_STAT(TEXT("MediaCapture RenderThread LockResource"), STAT_MediaCapture_RenderThread_LockResource, STATGROUP_Media);
 DECLARE_CYCLE_STAT(TEXT("MediaCapture RenderThread CPU Capture Callback"), STAT_MediaCapture_RenderThread_CaptureCallback, STATGROUP_Media);
@@ -110,7 +123,6 @@ FMediaCaptureOptions::FMediaCaptureOptions()
 	, bAutostopOnCapture(false)
 	, NumberOfFramesToCapture(-1)
 {
-
 }
 
 void FMediaCaptureOptions::PostSerialize(const FArchive& Ar)
@@ -519,6 +531,12 @@ bool UMediaCapture::UpdateTextureRenderTarget2D(UTextureRenderTarget2D* InRender
 
 void UMediaCapture::StopCapture(bool bAllowPendingFrameToBeProcess)
 {
+	if (bWaitingForStopCapture)
+	{
+		// This capture is already in the process of stopping.
+		return;
+	}
+
 	if (!IsInGameThread())
 	{
 		TWeakObjectPtr<UMediaCapture> Self = this;
@@ -564,7 +582,17 @@ void UMediaCapture::StopCapture(bool bAllowPendingFrameToBeProcess)
 
 		if (GetState() != EMediaCaptureState::Stopped)
 		{
-			SetState(EMediaCaptureState::Stopped);
+			if (SupportsStopCapture_Future())
+			{
+				SetState(EMediaCaptureState::StopRequested);
+				// Prevent reentry and subsequent calls to StopCapture with bAllowPendingFrameToBeProcess = false to set the state to Stopped
+				// while the TFuture isn't ready.
+				bWaitingForStopCapture = true;
+			}
+			else
+			{
+				SetState(EMediaCaptureState::Stopped);
+			}
 
 			FCoreDelegates::OnBeginFrame.RemoveAll(this);
 			FCoreDelegates::OnEndFrame.RemoveAll(this);
@@ -574,29 +602,31 @@ void UMediaCapture::StopCapture(bool bAllowPendingFrameToBeProcess)
 				FlushRenderingCommands();
 			}
 			
-			StopCaptureImpl(bAllowPendingFrameToBeProcess);
-
-			SetCaptureAudioDevice(FAudioDeviceHandle());
+			if (SupportsStopCapture_Future())
+			{
+				TWeakObjectPtr<UMediaCapture> Self = this;
 
-			if (CaptureSource)
+				TRACE_BOOKMARK(TEXT("UMediaCapture::InitiateStopCapture"));
+				StopCapture_Future().Then([Self](TFuture<void>) mutable
+				{
+					AsyncTask(ENamedThreads::GameThread, [Self]
+					{
+						// Setting bEvenifPendingKill because this might be called as part of BeginDestroy,
+						// but we still want to cleanup card connections.
+						if (TStrongObjectPtr<UMediaCapture> Capture = Self.Pin(/*bEvenIfPendingKill=*/true)) 
+						{
+							Capture->SetState(EMediaCaptureState::Stopped);
+							Capture->PostCaptureCleanup();
+							TRACE_BOOKMARK(TEXT("UMediaCapture::CompleteStopCapture"));
+						}
+					});
+				});
+			}
+			else
 			{
-				constexpr bool bFlushRenderingCommands = false;
-				CaptureSource->ResetSourceBufferSize(bFlushRenderingCommands);
-				CaptureSource.Reset();
+				StopCaptureImpl(false);
+				PostCaptureCleanup();
 			}
-
-			DesiredSize = FIntPoint(1280, 720);
-			DesiredPixelFormat = EPixelFormat::PF_A2B10G10R10;
-			DesiredOutputSize = FIntPoint(1280, 720);
-			DesiredOutputPixelFormat = EPixelFormat::PF_A2B10G10R10;
-			DesiredCaptureOptions = FMediaCaptureOptions();
-			ConversionOperation = EMediaCaptureConversionOperation::NONE;
-			SyncPointWatcher.Reset();
-			ViewExtension.Reset();
-			CaptureRenderPipeline.Reset();
-			bLoggedOutputNotReady = false;
-			
-			MediaOutputName.Reset();
 		}
 	}
 }
@@ -1407,6 +1437,13 @@ bool UMediaCapture::ProcessCapture_RenderThread(const TSharedPtr<UE::MediaCaptur
 
 	TRACE_CPUPROFILER_EVENT_SCOPE_TEXT(*FString::Printf(TEXT("Process Capture Render Thread Frame %d"), FrameNumber % 10));
 
+	if (bSimulateError)
+	{
+		Args.MediaCapture->SetState(EMediaCaptureState::Error);
+		bSimulateError = false;
+		return false;
+	}
+
 	if (CapturingFrame)
 	{
 		bool bHasCaptureSuceeded = false;
@@ -1690,6 +1727,35 @@ int32 UMediaCapture::GetBytesPerPixel(EPixelFormat InPixelFormat)
 	}
 }
 
+void UMediaCapture::PostCaptureCleanup()
+{
+	// Clean up implementation resources.
+	PostCaptureCleanupImpl();
+
+	SetCaptureAudioDevice(FAudioDeviceHandle());
+
+	if (CaptureSource)
+	{
+		constexpr bool bFlushRenderingCommands = false;
+		CaptureSource->ResetSourceBufferSize(bFlushRenderingCommands);
+		CaptureSource.Reset();
+	}
+
+	DesiredSize = FIntPoint(1280, 720);
+	DesiredPixelFormat = EPixelFormat::PF_A2B10G10R10;
+	DesiredOutputSize = FIntPoint(1280, 720);
+	DesiredOutputPixelFormat = EPixelFormat::PF_A2B10G10R10;
+	DesiredCaptureOptions = FMediaCaptureOptions();
+	ConversionOperation = EMediaCaptureConversionOperation::NONE;
+	SyncPointWatcher.Reset();
+	ViewExtension.Reset();
+	CaptureRenderPipeline.Reset();
+	bLoggedOutputNotReady = false;
+	MediaOutputName.Reset();
+
+	bWaitingForStopCapture = false;
+}
+
 /* namespace MediaCaptureDetails implementation
 *****************************************************************************/
 namespace MediaCaptureDetails
diff --git a/Engine/Plugins/Media/MediaIOFramework/Source/MediaIOCore/Public/MediaCapture.h b/Engine/Plugins/Media/MediaIOFramework/Source/MediaIOCore/Public/MediaCapture.h
--- a/Engine/Plugins/Media/MediaIOFramework/Source/MediaIOCore/Public/MediaCapture.h
+++ b/Engine/Plugins/Media/MediaIOFramework/Source/MediaIOCore/Public/MediaCapture.h
@@ -459,6 +459,28 @@ protected:
 	
 	virtual void StopCaptureImpl(bool bAllowPendingFrameToBeProcess) { }
 
+	/** Returns whether the capture implementation supports immediate (blocking) StopCapture.*/
+	virtual bool SupportsStopCapture_Future() const { return false; }
+
+	/**
+	 * Version of StopCaptureImpl that returns a future.
+	 * Override this in your capture implementation to allow MediaCapture to wait on the StopCapture to complete.
+	 * @note: Must be called from Game thread.
+	 */
+	virtual TFuture<void> StopCapture_Future()
+	{
+		ensureMsgf(false, TEXT("Attempted to call StopCapture_Future without checking if the implementation supports it via SupportsImmediateStopCapture."));
+		return MakeFulfilledPromise<void>().GetFuture();
+	}
+
+	/**
+	 * Override this to cleanup resources in the capture implementation after StopCapture is done. 
+	 * @note: Will be invoked from the game thread.
+	 */
+	virtual void PostCaptureCleanupImpl()
+	{
+	}
+
 	//~ DMA Functions implemented in child classes, Temporary until AJA and BM middleman are migrated to a module.
 	virtual void LockDMATexture_RenderThread(FTextureRHIRef InTexture) {}
 	virtual void UnlockDMATexture_RenderThread(FTextureRHIRef InTexture) {}
@@ -678,6 +700,9 @@ private:
 	/** Returns bytes per pixel based on pixel format */
 	static UE_API int32 GetBytesPerPixel(EPixelFormat InPixelFormat);
 
+	/** Cleanup config and resources after StopCaptureImpl is invoked. */
+	void PostCaptureCleanup();
+
 protected:
 
 	/** MediaOutput associated with this capture */
@@ -720,6 +745,9 @@ private:
 	/** Used to avoid repeatedly logging while the underlying sdk is not ready to output. */
 	bool bLoggedOutputNotReady = false;
 
+	/** This indicates the capture is in the process of stopping and is waiting for StopCapture_Future's result to be ready. */
+	bool bWaitingForStopCapture = false;
+
 	bool bViewportHasFixedViewportSize = false;
 	std::atomic<bool> bPendingOutputResourcesInitialized;
 	std::atomic<bool> bSyncHandlersInitialized;
