MediaCapture: Fix frame number format specifier.

Also shortened the bookmark label.

#rnx

diff --git a/Engine/Plugins/Media/MediaIOFramework/Source/MediaIOCore/Private/MediaCapture.cpp b/Engine/Plugins/Media/MediaIOFramework/Source/MediaIOCore/Private/MediaCapture.cpp
--- a/Engine/Plugins/Media/MediaIOFramework/Source/MediaIOCore/Private/MediaCapture.cpp
+++ b/Engine/Plugins/Media/MediaIOFramework/Source/MediaIOCore/Private/MediaCapture.cpp
@@ -1347,7 +1347,7 @@ void UMediaCapture::AddFrameInfoForCaptureImmediate(const uint64 FrameNumber)
 			return;
 		}
 
-		TRACE_BOOKMARK(TEXT("MediaCapture AddFrameInfoForCaptureImmediate for frame %d"), FrameNumber);
+		TRACE_BOOKMARK(TEXT("MediaCapture GT push frame %llu"), FrameNumber);
 		CaptureDataQueue.Insert(MoveTemp(CaptureData), 0);
 	}
 }
