Media IO - Fix scene view extension sometimes not working in editor because of a mismatched render target in the capture scene view extension

diff --git a/Engine/Plugins/Media/MediaIOFramework/Source/MediaIOCore/Private/MediaCapture.cpp b/Engine/Plugins/Media/MediaIOFramework/Source/MediaIOCore/Private/MediaCapture.cpp
--- a/Engine/Plugins/Media/MediaIOFramework/Source/MediaIOCore/Private/MediaCapture.cpp
+++ b/Engine/Plugins/Media/MediaIOFramework/Source/MediaIOCore/Private/MediaCapture.cpp
@@ -321,12 +321,30 @@ bool UMediaCapture::StartSourceCapture(TSharedPtr<UE::MediaCapture::Private::FCa
 	{
 		//no lock required the command on the render thread is not active yet
 		CaptureSource = MoveTemp(InSource);
-		
+
 		if (DesiredCaptureOptions.CapturePhase != EMediaCapturePhase::EndFrame)
 		{
 			/** If we are running with color conversion, we need our extension to be executed before the OCIO extension is applied to the viewport. */
 			const int32 Priority =  DesiredCaptureOptions.ColorConversionSettings.IsValid() ? OPENCOLORIO_SCENE_VIEW_EXTENSION_PRIORITY + 1 : 0;
-			ViewExtension = FSceneViewExtensions::NewExtension<FMediaCaptureSceneViewExtension>(this, DesiredCaptureOptions.CapturePhase, Priority);
+
+			FRenderTarget* RenderTarget = nullptr;
+			if (CaptureSource->GetSceneViewport())
+			{
+				RenderTarget = CaptureSource->GetSceneViewport()->GetViewport();
+			}
+			else if (CaptureSource->GetRenderTarget())
+			{
+				RenderTarget = CaptureSource->GetRenderTarget()->GameThread_GetRenderTargetResource();
+			}
+
+			if (RenderTarget) 
+			{
+				ViewExtension = FSceneViewExtensions::NewExtension<FMediaCaptureSceneViewExtension>(this, DesiredCaptureOptions.CapturePhase, Priority, RenderTarget);
+			}
+			else
+			{
+				UE_LOG(LogMediaIOCore, Warning, TEXT("Could not acquire a render target to initialize the scene view extension. Consider changing the capture phase to End Frame."));
+			}
 		}
 
 		if (ViewExtension)
@@ -1361,6 +1379,11 @@ void UMediaCapture::OnBeginFrame_GameThread()
 		SetState(EMediaCaptureState::Error);
 	}
 
+	if (CaptureSource && !CaptureSource->IsValid())
+	{
+		SetState(EMediaCaptureState::Error);
+	}
+
 	if (GetState() == EMediaCaptureState::Error)
 	{
 		StopCapture(false);
diff --git a/Engine/Plugins/Media/MediaIOFramework/Source/MediaIOCore/Private/MediaCaptureSceneViewExtension.h b/Engine/Plugins/Media/MediaIOFramework/Source/MediaIOCore/Private/MediaCaptureSceneViewExtension.h
--- a/Engine/Plugins/Media/MediaIOFramework/Source/MediaIOCore/Private/MediaCaptureSceneViewExtension.h
+++ b/Engine/Plugins/Media/MediaIOFramework/Source/MediaIOCore/Private/MediaCaptureSceneViewExtension.h
@@ -24,13 +24,13 @@ class FMediaCaptureSceneViewExtension : public FSceneViewExtensionBase
 {
 public:
 	
-	FMediaCaptureSceneViewExtension(const FAutoRegister& InAutoRegister, UMediaCapture* InMediaCapture, EMediaCapturePhase InCapturePhase, int32 InPriority)
+	FMediaCaptureSceneViewExtension(const FAutoRegister& InAutoRegister, UMediaCapture* InMediaCapture, EMediaCapturePhase InCapturePhase, int32 InPriority, FRenderTarget* InRenderTarget)
 		: FSceneViewExtensionBase(InAutoRegister)
 		, WeakCapture(InMediaCapture)
 		, CapturePhase(InCapturePhase)
 		, Priority(InPriority)
+		, CurrentRenderTarget(InRenderTarget)
 	{
-		
 	}
 
 	//~ Begin FSceneViewExtensionBase Interface
@@ -41,7 +41,7 @@ public:
 
 	virtual void SetupView(FSceneViewFamily& InViewFamily, FSceneView& InView) override 
 	{
-		if (bValidPhase)
+		if (bValidPhase && CurrentRenderTarget == InView.Family->RenderTarget)
 		{
 			// Copied from PostProcessing.h
 			if (InView.GetFeatureLevel() >= ERHIFeatureLevel::SM5)
@@ -59,7 +59,7 @@ public:
 				bPostProcessingEnabled = InView.Family->EngineShowFlags.PostProcessing && !InView.Family->EngineShowFlags.ShaderComplexity && IsMobileHDR();
 			}
 
-			if (CapturePhase != EMediaCapturePhase::BeforePostProcessing && CapturePhase != EMediaCapturePhase::EndFrame)
+			if (CapturePhase != EMediaCapturePhase::BeforePostProcessing && CapturePhase != EMediaCapturePhase::EndFrame && CapturePhase != EMediaCapturePhase::PostRender)
 			{
 				if (!bPostProcessingEnabled)
 				{
@@ -71,14 +71,17 @@ public:
 		}
 	};
 
-	virtual void SubscribeToPostProcessingPass(EPostProcessingPass PassId, const FSceneView& View, FAfterPassCallbackDelegateArray& InOutPassCallbacks, bool bIsPassEnabled) override
+	virtual void SubscribeToPostProcessingPass(EPostProcessingPass PassId, const FSceneView& InView, FAfterPassCallbackDelegateArray& InOutPassCallbacks, bool bIsPassEnabled) override
 	{
-		if ((CapturePhase == EMediaCapturePhase::AfterMotionBlur && PassId == EPostProcessingPass::MotionBlur)
-			|| (CapturePhase == EMediaCapturePhase::AfterToneMap && PassId == EPostProcessingPass::Tonemap)
-			|| (CapturePhase == EMediaCapturePhase::AfterFXAA && PassId == EPostProcessingPass::FXAA)
-			|| (CapturePhase == EMediaCapturePhase::BeforePostProcessing && PassId == EPostProcessingPass::SSRInput))
+		if (CurrentRenderTarget == InView.Family->RenderTarget)
 		{
-			InOutPassCallbacks.Add(FAfterPassCallbackDelegate::CreateSP(this, &FMediaCaptureSceneViewExtension::PostProcessCallback_RenderThread));
+			if ((CapturePhase == EMediaCapturePhase::AfterMotionBlur && PassId == EPostProcessingPass::MotionBlur)
+				|| (CapturePhase == EMediaCapturePhase::AfterToneMap && PassId == EPostProcessingPass::Tonemap)
+				|| (CapturePhase == EMediaCapturePhase::AfterFXAA && PassId == EPostProcessingPass::FXAA)
+				|| (CapturePhase == EMediaCapturePhase::BeforePostProcessing && PassId == EPostProcessingPass::SSRInput))
+			{
+				InOutPassCallbacks.Add(FAfterPassCallbackDelegate::CreateSP(this, &FMediaCaptureSceneViewExtension::PostProcessCallback_RenderThread));
+			}
 		}
 	}
 
@@ -89,7 +92,7 @@ public:
 
 	virtual void PostRenderView_RenderThread(FRDGBuilder& GraphBuilder, FSceneView& InView) override
 	{
-		if (CapturePhase == EMediaCapturePhase::PostRender)
+		if (CapturePhase == EMediaCapturePhase::PostRender && CurrentRenderTarget == InView.Family->RenderTarget)
 		{
 			TRACE_CPUPROFILER_EVENT_SCOPE(MediaCaptureExtensionCallback);
 
@@ -140,5 +143,7 @@ private:
 	bool bValidPhase = true;
 	FString LastErrorMessage;
 	int32 Priority = 0;
+	/** Render target that should be considered by this scene view extension. */
+	FRenderTarget* CurrentRenderTarget = nullptr;
 };
 
diff --git a/Engine/Plugins/Media/MediaIOFramework/Source/MediaIOCore/Private/MediaCaptureSources.h b/Engine/Plugins/Media/MediaIOFramework/Source/MediaIOCore/Private/MediaCaptureSources.h
--- a/Engine/Plugins/Media/MediaIOFramework/Source/MediaIOCore/Private/MediaCaptureSources.h
+++ b/Engine/Plugins/Media/MediaIOFramework/Source/MediaIOCore/Private/MediaCaptureSources.h
@@ -155,6 +155,8 @@ namespace UE::MediaCapture::Private
 		virtual bool PostInitialize() = 0;
 		virtual EMediaCaptureSourceType GetSourceType() = 0;
 		virtual FIntPoint GetSize() = 0;
+		/** Returns whether the underlying captured resource is still valid. */
+		virtual bool IsValid() const = 0;
 		
 		virtual ECaptureType GetCaptureType() const
 		{
@@ -309,6 +311,11 @@ namespace UE::MediaCapture::Private
 		{
 			return WeakViewport.Pin();
 		}
+
+		virtual bool IsValid() const override
+		{
+			return WeakViewport.IsValid();
+		}
 		
 		TWeakPtr<FSceneViewport> WeakViewport;
 	};
@@ -347,6 +354,12 @@ namespace UE::MediaCapture::Private
 			return true;
 		}
 
+		virtual bool IsValid() const override
+		{
+			// We don't keep a handle to the resource, and it may change on every frame so we always return true to prevent stopping the capture.
+			return true;
+		}
+
 		FRHICaptureResourceDescription ResourceDescription;
 	};
 
@@ -418,6 +431,11 @@ namespace UE::MediaCapture::Private
 			return RenderTarget.Get();
 		}
 
+		virtual bool IsValid() const override
+		{
+			return RenderTarget.IsValid();
+		}
+
 		TWeakObjectPtr<UTextureRenderTarget2D> RenderTarget;
 	};
 }
