MediaCapture: Tentative null pointer access crash to either MediaCapture or CaptureSource.

This should now be avoided:
* Weak this pointers pinned into strong pointers.
* CaptureSource shared pointer is copied into lambda while valid to maintain its lifetime.

diff --git a/Engine/Plugins/Media/MediaIOFramework/Source/MediaIOCore/Private/MediaCapture.cpp b/Engine/Plugins/Media/MediaIOFramework/Source/MediaIOCore/Private/MediaCapture.cpp
--- a/Engine/Plugins/Media/MediaIOFramework/Source/MediaIOCore/Private/MediaCapture.cpp
+++ b/Engine/Plugins/Media/MediaIOFramework/Source/MediaIOCore/Private/MediaCapture.cpp
@@ -1114,24 +1114,29 @@ void UMediaCapture::PrepareAndDispatchCapture_GameThread(const TSharedPtr<UE::Me
 	{
 		FScopeLock Lock(&AccessingCapturingSource);
 
-		FIntPoint InDesiredSize = DesiredSize;
-		UMediaCapture* InMediaCapture = this;
-
-		if (CaptureSource)
+		if (CaptureSource.IsValid())
 		{
 			++WaitingForRenderCommandExecutionCounter;
 
 			const FRHIGPUMask SourceGPUMask = ValidSourceGPUMask;
 
-			UE_LOG(LogMediaIOCore, Verbose, TEXT("[%s] - Queuing frame to capture %d"), *InMediaCapture->MediaOutputName, CapturingFrame ? CapturingFrame->FrameId : -1);
+			TWeakObjectPtr<UMediaCapture> WeakThis = this;
+
+			UE_LOG(LogMediaIOCore, Verbose, TEXT("[%s] - Queuing frame to capture %d"), *MediaOutputName, CapturingFrame ? CapturingFrame->FrameId : -1);
 
 			// RenderCommand to be executed on the RenderThread
 			ENQUEUE_RENDER_COMMAND(FMediaOutputCaptureFrameCreateTexture)(
-				[InMediaCapture, CapturingFrame, InDesiredSize, SourceGPUMask](FRHICommandListImmediate& RHICmdList)
+				[WeakThis, CaptureSource_RT = CaptureSource, CapturingFrame, InDesiredSize = DesiredSize, SourceGPUMask](FRHICommandListImmediate& RHICmdList)
 			{
 				SCOPED_GPU_MASK(RHICmdList, SourceGPUMask);
 
-				TSharedPtr<FCaptureFrame> NextPending = InMediaCapture->FrameManager->PeekNextPending<FCaptureFrame>();
+				TStrongObjectPtr<UMediaCapture> This = WeakThis.Pin();
+				if (!This.IsValid())
+				{
+					return;
+				}
+
+				TSharedPtr<FCaptureFrame> NextPending = This->FrameManager->PeekNextPending<FCaptureFrame>();
 				if (NextPending && CapturingFrame)
 				{
 					ensure(NextPending->FrameId != CapturingFrame->FrameId);
@@ -1141,31 +1146,31 @@ void UMediaCapture::PrepareAndDispatchCapture_GameThread(const TSharedPtr<UE::Me
 				{
 					FRDGBuilder GraphBuilder(RHICmdList);
 					
-					FTextureRHIRef SourceTexture = InMediaCapture->CaptureSource->GetSourceTextureForInput_RenderThread();
+					FTextureRHIRef SourceTexture = CaptureSource_RT->GetSourceTextureForInput_RenderThread();
 					
 					FCaptureFrameArgs CaptureArgs{ GraphBuilder };
-					CaptureArgs.MediaCapture = InMediaCapture;
+					CaptureArgs.MediaCapture = This.Get();
 					CaptureArgs.DesiredSize = InDesiredSize;
 					CaptureArgs.ResourceToCapture = MoveTemp(SourceTexture);
 
-					InMediaCapture->ProcessCapture_RenderThread(CapturingFrame, CaptureArgs);
-					UE_LOG(LogMediaIOCore, Verbose, TEXT("[%s] - Captured frame %d"), *InMediaCapture->MediaOutputName, CapturingFrame ? CapturingFrame->FrameId : -1);
+					This->ProcessCapture_RenderThread(CapturingFrame, CaptureArgs);
+					UE_LOG(LogMediaIOCore, Verbose, TEXT("[%s] - Captured frame %d"), *This->MediaOutputName, CapturingFrame ? CapturingFrame->FrameId : -1);
 					
 					GraphBuilder.Execute();
 				}
 				
-				if (!InMediaCapture->UseExperimentalScheduling())
+				if (!This->UseExperimentalScheduling())
 				{
 					//Process the next pending frame
-					UE_LOG(LogMediaIOCore, Verbose, TEXT("[%s] - Processing pending frame %d"), *InMediaCapture->MediaOutputName, NextPending ? NextPending->FrameId : -1);
+					UE_LOG(LogMediaIOCore, Verbose, TEXT("[%s] - Processing pending frame %d"), *This->MediaOutputName, NextPending ? NextPending->FrameId : -1);
 					if (NextPending)
 					{
-						InMediaCapture->ProcessReadyFrame_RenderThread(RHICmdList, InMediaCapture, NextPending);
+						This->ProcessReadyFrame_RenderThread(RHICmdList, This.Get(), NextPending);
 					}
 				}
 
 				// Whatever happens, we want to decrement our counter to track enqueued commands
-				--InMediaCapture->WaitingForRenderCommandExecutionCounter;
+				--This->WaitingForRenderCommandExecutionCounter;
 			});
 
 			//If auto-stopping, count the number of frame captures requested and stop when reaching 0.
@@ -1207,10 +1212,16 @@ void UMediaCapture::InitializeOutputResources(int32 InNumberOfBuffers)
 	bSyncHandlersInitialized = false;
 	NumberOfCaptureFrame = InNumberOfBuffers;
 
-	UMediaCapture* This = this;
+	TWeakObjectPtr<UMediaCapture> WeakThis = this;
 	ENQUEUE_RENDER_COMMAND(MediaOutputCaptureFrameCreateResources)(
-	[This](FRHICommandListImmediate& RHICmdList)
+	[WeakThis](FRHICommandListImmediate& RHICmdList)
 		{
+			TStrongObjectPtr<UMediaCapture> This = WeakThis.Pin();
+			if (!This.IsValid())
+			{
+				return;
+			}
+
 			for (int32 Index = 0; Index < This->NumberOfCaptureFrame; ++Index)
 			{
 				if (This->DesiredOutputResourceType == EMediaCaptureResourceType::Texture)
@@ -1565,8 +1576,8 @@ bool UMediaCapture::ProcessCapture_RenderThread(const TSharedPtr<UE::MediaCaptur
 				TWeakObjectPtr<UMediaCapture> Self = this;
 				AsyncTask(ENamedThreads::GameThread, [Self]
 				{
-					UMediaCapture* MediaCapture = Self.Get();
-					if (UObjectInitialized() && MediaCapture)
+					TStrongObjectPtr<UMediaCapture> MediaCapture = Self.Pin();
+					if (UObjectInitialized() && MediaCapture.IsValid())
 					{
 						MediaCapture->RestartCapture();
 					}
diff --git a/Engine/Plugins/Media/MediaIOFramework/Source/MediaIOCore/Public/MediaCapture.h b/Engine/Plugins/Media/MediaIOFramework/Source/MediaIOCore/Public/MediaCapture.h
--- a/Engine/Plugins/Media/MediaIOFramework/Source/MediaIOCore/Public/MediaCapture.h
+++ b/Engine/Plugins/Media/MediaIOFramework/Source/MediaIOCore/Public/MediaCapture.h
@@ -789,7 +789,8 @@ private:
 	/** Watcher thread looking after end of capturing frames */
 	TPimplPtr<UE::MediaCaptureData::FSyncPointWatcher> SyncPointWatcher;
 
-	TSharedPtr<UE::MediaCapture::Private::FCaptureSource> CaptureSource;
+	/** Capture source */
+	TSharedPtr<UE::MediaCapture::Private::FCaptureSource, ESPMode::ThreadSafe> CaptureSource;
 
 	/** Holds a view extension that callbacks  */
 	TSharedPtr<class FMediaCaptureSceneViewExtension> ViewExtension;
