Media IO - Fix Media Source not being able to unset AutoDetect in Media Profile

This looks to be caused by a limitation of customizing Instanced UObjects, where at some point, the property node for the configuration property would have no outer objects set. To cirvumvent this, I've added a list to keep track of customized objects while the customization exists.

#rnx

diff --git a/Engine/Plugins/Media/MediaIOFramework/Source/MediaIOEditor/Private/Customizations/MediaIOConfigurationCustomization.cpp b/Engine/Plugins/Media/MediaIOFramework/Source/MediaIOEditor/Private/Customizations/MediaIOConfigurationCustomization.cpp
--- a/Engine/Plugins/Media/MediaIOFramework/Source/MediaIOEditor/Private/Customizations/MediaIOConfigurationCustomization.cpp
+++ b/Engine/Plugins/Media/MediaIOFramework/Source/MediaIOEditor/Private/Customizations/MediaIOConfigurationCustomization.cpp
@@ -2,6 +2,7 @@
 
 #include "Customizations/MediaIOConfigurationCustomization.h"
 
+#include "Algo/Transform.h"
 #include "DetailWidgetRow.h"
 #include "DetailCategoryBuilder.h"
 #include "DetailLayoutBuilder.h"
@@ -29,6 +30,12 @@ TSharedRef<IPropertyTypeCustomization> FMediaIOConfigurationCustomization::MakeI
 	return MakeShared<FMediaIOConfigurationCustomization>();
 }
 
+void FMediaIOConfigurationCustomization::OnCustomizeObjects(const TArray<UObject*>& InCustomizedObjects)
+{
+	Objects.Reserve(InCustomizedObjects.Num());
+	Algo::Transform(InCustomizedObjects, Objects, [](UObject* Object) { return TWeakObjectPtr<>{ Object }; });
+}
+
 TAttribute<FText> FMediaIOConfigurationCustomization::GetContentText()
 {
 	FMediaIOConfiguration* Value = GetPropertyValueFromPropertyHandle<FMediaIOConfiguration>();
@@ -219,23 +226,26 @@ bool FMediaIOConfigurationCustomization::IsAutoDetected() const
 {
 	bool bAutoDetectTimecode = true;
 
-	for (UObject* CustomizedObject : GetCustomizedObjects())
+	for (const TWeakObjectPtr<UObject>& WeakObject : Objects)
 	{
-		if (const UTimeSynchronizableMediaSource* Source = Cast<UTimeSynchronizableMediaSource>(CustomizedObject))
+		if (TStrongObjectPtr<UObject> CustomizedObject = WeakObject.Pin())
 		{
-			if (!Source->bAutoDetectInput)
+			if (const UTimeSynchronizableMediaSource* Source = Cast<UTimeSynchronizableMediaSource>(CustomizedObject.Get()))
 			{
-				bAutoDetectTimecode = false;
-				break;
+				if (!Source->bAutoDetectInput)
+				{
+					bAutoDetectTimecode = false;
+					break;
+				}
 			}
-		}
 
-		if (const UGenlockedCustomTimeStep* CustomTimeStep = Cast<UGenlockedCustomTimeStep>(CustomizedObject))
-		{
-			if (!CustomTimeStep->bAutoDetectFormat)
+			if (const UGenlockedCustomTimeStep* CustomTimeStep = Cast<UGenlockedCustomTimeStep>(CustomizedObject.Get()))
 			{
-				bAutoDetectTimecode = false;
-				break;
+				if (!CustomTimeStep->bAutoDetectFormat)
+				{
+					bAutoDetectTimecode = false;
+					break;
+				}
 			}
 		}
 	}
@@ -247,18 +257,21 @@ void FMediaIOConfigurationCustomization::SetIsAutoDetected(bool Value)
 {
 	FScopedTransaction Transaction{ LOCTEXT("MediaIOAutoDetectTimecode", "Auto Detect Timecode") };
 
-	for (UObject* CustomizedObject : GetCustomizedObjects())
+	for (const TWeakObjectPtr<UObject>& WeakObject : Objects)
 	{
-		if (UTimeSynchronizableMediaSource* Source = Cast<UTimeSynchronizableMediaSource>(CustomizedObject))
+		if (TStrongObjectPtr<UObject> CustomizedObject = WeakObject.Pin())
 		{
-			Source->Modify();
-			Source->bAutoDetectInput = Value;
-		}
+			if (UTimeSynchronizableMediaSource* Source = Cast<UTimeSynchronizableMediaSource>(CustomizedObject.Get()))
+			{
+				Source->Modify();
+				Source->bAutoDetectInput = Value;
+			}
 
-		if (UGenlockedCustomTimeStep* CustomTimeStep = Cast<UGenlockedCustomTimeStep>(CustomizedObject))
-		{
-			CustomTimeStep->Modify();
-			CustomTimeStep->bAutoDetectFormat = Value;
+			if (UGenlockedCustomTimeStep* CustomTimeStep = Cast<UGenlockedCustomTimeStep>(CustomizedObject.Get()))
+			{
+				CustomTimeStep->Modify();
+				CustomTimeStep->bAutoDetectFormat = Value;
+			}
 		}
 	}
 }
diff --git a/Engine/Plugins/Media/MediaIOFramework/Source/MediaIOEditor/Private/Customizations/MediaIOCustomizationBase.cpp b/Engine/Plugins/Media/MediaIOFramework/Source/MediaIOEditor/Private/Customizations/MediaIOCustomizationBase.cpp
--- a/Engine/Plugins/Media/MediaIOFramework/Source/MediaIOEditor/Private/Customizations/MediaIOCustomizationBase.cpp
+++ b/Engine/Plugins/Media/MediaIOFramework/Source/MediaIOEditor/Private/Customizations/MediaIOCustomizationBase.cpp
@@ -44,11 +44,14 @@ void FMediaIOCustomizationBase::CustomizeHeader(TSharedRef<IPropertyHandle> InPr
 			//DeviceProviderName = *MediaProperty->GetPropertyClass()->GetMetaData(NAME_MediaIOCustomLayout);
 		//}
 
+		TArray<UObject*> OuterObjects;
+		MediaProperty->GetOuterObjects(OuterObjects);
+
+		// Call into child classes as soon as we have access to the uobjects.
+		OnCustomizeObjects(OuterObjects);
+
 		if (DeviceProviderName == NAME_None)
 		{
-			TArray<UObject*> OuterObjects;
-			MediaProperty->GetOuterObjects(OuterObjects);
-
 			for (UObject* Obj : OuterObjects)
 			{
 				FName FoundMetadata = FMediaIOCustomizationBaseHelper::TraverseClassHierarchy(Obj->GetClass(), NAME_MediaIOCustomLayout);
diff --git a/Engine/Plugins/Media/MediaIOFramework/Source/MediaIOEditor/Public/Customizations/MediaIOConfigurationCustomization.h b/Engine/Plugins/Media/MediaIOFramework/Source/MediaIOEditor/Public/Customizations/MediaIOConfigurationCustomization.h
--- a/Engine/Plugins/Media/MediaIOFramework/Source/MediaIOEditor/Public/Customizations/MediaIOConfigurationCustomization.h
+++ b/Engine/Plugins/Media/MediaIOFramework/Source/MediaIOEditor/Public/Customizations/MediaIOConfigurationCustomization.h
@@ -16,6 +16,10 @@ class FMediaIOConfigurationCustomization : public FMediaIOCustomizationBase
 public:
 	static UE_API TSharedRef<IPropertyTypeCustomization> MakeInstance();
 
+protected:
+	//~ FMediaIOCustomizationBase interface
+	virtual void OnCustomizeObjects(const TArray<UObject*>& InCustomizedObjects) override;
+
 private:
 	UE_API virtual TAttribute<FText> GetContentText() override;
 	UE_API virtual TSharedRef<SWidget> HandleSourceComboButtonMenuContent() override;
@@ -32,6 +36,8 @@ private:
 	TWeakPtr<SWidget> PermutationSelector;
 	FMediaIOConfiguration SelectedConfiguration;
 	bool bAutoDetectFormat = false;
+	/** List of customized objects. */
+	TArray<TWeakObjectPtr<UObject>> Objects;
 };
 
 #undef UE_API
diff --git a/Engine/Plugins/Media/MediaIOFramework/Source/MediaIOEditor/Public/Customizations/MediaIOCustomizationBase.h b/Engine/Plugins/Media/MediaIOFramework/Source/MediaIOEditor/Public/Customizations/MediaIOCustomizationBase.h
--- a/Engine/Plugins/Media/MediaIOFramework/Source/MediaIOEditor/Public/Customizations/MediaIOCustomizationBase.h
+++ b/Engine/Plugins/Media/MediaIOFramework/Source/MediaIOEditor/Public/Customizations/MediaIOCustomizationBase.h
@@ -50,6 +50,9 @@ protected:
 	TSharedPtr<IPropertyHandle> GetMediaProperty() const { return MediaProperty; };
 	UE_API TArray<UObject*> GetCustomizedObjects() const;
 
+	/** Used to give a chance to the child implementations to get access to the customized objects as soon as possible. */
+	virtual void OnCustomizeObjects(const TArray<UObject*>& InCustomizedObjects) {}
+
 	FName DeviceProviderName;
 
 private:
