Media IO - When using ResizeInRenderPass, fix output getting resized even if the input resolution matched the desired output size.

diff --git a/Engine/Plugins/Media/MediaIOFramework/Source/MediaIOCore/Private/MediaCaptureRenderPass.cpp b/Engine/Plugins/Media/MediaIOFramework/Source/MediaIOCore/Private/MediaCaptureRenderPass.cpp
--- a/Engine/Plugins/Media/MediaIOFramework/Source/MediaIOCore/Private/MediaCaptureRenderPass.cpp
+++ b/Engine/Plugins/Media/MediaIOFramework/Source/MediaIOCore/Private/MediaCaptureRenderPass.cpp
@@ -247,6 +247,13 @@ namespace UE::MediaCapture
 		ResamplePass.Name = "Resample";
 		ResamplePass.OutputType = ERDGViewableResourceType::Texture;
 		ResamplePass.InitializePassOutputDelegate = FRenderPass::FInitializePassOutput::CreateStatic(&InitializeResamplePassOutputTexture);
+		ResamplePass.ShouldExecutePassDelegate = FRenderPass::FShouldExecutePass::CreateLambda([](const UE::MediaCaptureData::FCaptureFrameArgs& Args, FRDGViewableResource* InputResource)
+		{
+			check(InputResource->Type == ERDGViewableResourceType::Texture);
+			FRDGTextureRef SourceRGBTexture = static_cast<FRDGTextureRef>(InputResource);
+			return SourceRGBTexture->Desc.Extent != Args.DesiredSize;
+		});
+
 		ResamplePass.ExecutePassDelegate = FRenderPass::FExecutePass::CreateLambda([](const UE::MediaCaptureData::FCaptureFrameArgs& Args, const TSharedPtr<UE::MediaCaptureData::FCaptureFrame>& CapturingFrame, FRDGViewableResource* InputResource, FRDGViewableResource* OutputTexture)
 		{
 			check(InputResource->Type == ERDGViewableResourceType::Texture);
@@ -257,7 +264,6 @@ namespace UE::MediaCapture
 			GetCopyInfo(Args, CopyInfo, SizeU, SizeV);
 
 			const FBaseRenderPassArgs ColorArgs { static_cast<FRDGTextureRef>(InputResource), CopyInfo };
-				
 			AddTextureResamplePass(Args, ColorArgs, OutputTexture);
 		});
 
@@ -541,12 +547,20 @@ namespace UE::MediaCapture
 					return PreviousOutput;
 				}
 
-				FRDGViewableResource* PassOutput = RegisterPassOutputResource(CapturingFrame->RenderPassResources, GraphBuilder, RenderPass);
 
 				// Register pass rdg resource
-				RenderPass.ExecutePassDelegate.Execute(Args, CapturingFrame, PreviousOutput, PassOutput);
+				bool bExecutePass = true;
+				if (RenderPass.ShouldExecutePassDelegate.IsBound())
+				{
+					bExecutePass = RenderPass.ShouldExecutePassDelegate.Execute(Args, PreviousOutput);
+				}
 
-				PreviousOutput = PassOutput;
+				if (bExecutePass)
+				{
+					FRDGViewableResource* PassOutput = RegisterPassOutputResource(CapturingFrame->RenderPassResources, GraphBuilder, RenderPass);
+					RenderPass.ExecutePassDelegate.Execute(Args, CapturingFrame, PreviousOutput, PassOutput);
+					PreviousOutput = PassOutput;
+				}
 			}
 
 		}
diff --git a/Engine/Plugins/Media/MediaIOFramework/Source/MediaIOCore/Private/MediaCaptureRenderPass.h b/Engine/Plugins/Media/MediaIOFramework/Source/MediaIOCore/Private/MediaCaptureRenderPass.h
--- a/Engine/Plugins/Media/MediaIOFramework/Source/MediaIOCore/Private/MediaCaptureRenderPass.h
+++ b/Engine/Plugins/Media/MediaIOFramework/Source/MediaIOCore/Private/MediaCaptureRenderPass.h
@@ -84,6 +84,12 @@ namespace UE::MediaCapture
 		 */
 		FExecutePass ExecutePassDelegate;
 
+		/**
+		 * Delegate invoked to query whether the pass should be executed.
+		 */
+		DECLARE_DELEGATE_RetVal_TwoParams(bool, FShouldExecutePass, const UE::MediaCaptureData::FCaptureFrameArgs& /*Args*/, FRDGViewableResource* /*InputResource*/);
+		FShouldExecutePass ShouldExecutePassDelegate;
+
 		/** Name of the pass. */
 		FName Name;
 
