Media IO - Fix Media Capture taking multiple frames to start outputting

Add a CVar (MediaIO.PrerollFrames) to control whether to preroll frames.
Prerolling can be useful for SDI workflow, but might be undesirable when recording to a file since beginning frames will be skipped.

diff --git a/Engine/Plugins/Media/MediaIOFramework/Source/MediaIOCore/Private/MediaCapture.cpp b/Engine/Plugins/Media/MediaIOFramework/Source/MediaIOCore/Private/MediaCapture.cpp
--- a/Engine/Plugins/Media/MediaIOFramework/Source/MediaIOCore/Private/MediaCapture.cpp
+++ b/Engine/Plugins/Media/MediaIOFramework/Source/MediaIOCore/Private/MediaCapture.cpp
@@ -57,6 +57,12 @@ static TAutoConsoleVariable<int32> CVarMediaIOScheduleOnAnyThread(
 	TEXT("Whether to wait for resource readback in a separate thread. (Experimental)"),
 	ECVF_RenderThreadSafe);
 
+static TAutoConsoleVariable<int32> CVarMediaIOPrerollFrames(
+	TEXT("MediaIO.PrerollFrames"), 0,
+	TEXT("Whether to preroll frames during media capture. Enabling this may reduce latency by initializing GPU resources during the preroll phase"),
+	ECVF_RenderThreadSafe);
+
+
 static bool bSimulateError = false;
 
 FAutoConsoleCommand FMediaIOSetErrorState
@@ -1245,7 +1251,7 @@ void UMediaCapture::InitializeOutputResources(int32 InNumberOfBuffers)
 						NewFrame->ReadbackTexture = MakeUnique<FRHIGPUTextureReadback>(*FString::Printf(TEXT("MediaCaptureTextureReadback_%d"), Index));
 					}
 
-					if (This->GetDesiredCaptureOptions().bAutostopOnCapture)
+					if (This->GetDesiredCaptureOptions().bAutostopOnCapture || !CVarMediaIOPrerollFrames.GetValueOnRenderThread())
 					{
 						// Don't preroll if capturing a set number of frames since we discard the first frames.
 						NewFrame->bPrerollFrame = false;
@@ -1509,7 +1515,7 @@ bool UMediaCapture::ProcessCapture_RenderThread(const TSharedPtr<UE::MediaCaptur
 				check(!GetDesiredCaptureOptions().bAutostopOnCapture);
 			}
 
-			if (CapturingFrame->bPrerollFrame || GetDesiredCaptureOptions().bAutostopOnCapture)
+			if (CapturingFrame->bPrerollFrame || GetDesiredCaptureOptions().bAutostopOnCapture || !CVarMediaIOPrerollFrames.GetValueOnRenderThread())
 			{
 				Args.MediaCapture->CaptureRenderPipeline->InitializeResources_RenderThread(CapturingFrame, Args.GraphBuilder, Args.RDGResourceToCapture);
 			}
