GPUTextureTransfer - Remove unused GPUTextureTransfer code

#rnx

diff --git a/Engine/Plugins/Media/MediaIOFramework/Source/GPUTextureTransfer/Private/GPUTextureTransferModule.cpp b/Engine/Plugins/Media/MediaIOFramework/Source/GPUTextureTransfer/Private/GPUTextureTransferModule.cpp
--- a/Engine/Plugins/Media/MediaIOFramework/Source/GPUTextureTransfer/Private/GPUTextureTransferModule.cpp
+++ b/Engine/Plugins/Media/MediaIOFramework/Source/GPUTextureTransfer/Private/GPUTextureTransferModule.cpp
@@ -347,49 +347,6 @@ namespace UE::GPUTextureTransfer
 		}
 
 	} TextureTransfersWrapper;
-
-	ITextureTransfer* GetTextureTransfer(const UE::GPUTextureTransfer::FInitializeDMAArgs& Args)
-	{
-#if DVP_SUPPORTED_PLATFORM
-		const uint8_t RHIIndex = static_cast<uint8_t>(Args.RHI);
-
-		if (TextureTransfersWrapper.Transfers[RHIIndex])
-		{
-			return TextureTransfersWrapper.Transfers[RHIIndex];
-		}
-
-		ITextureTransfer* TextureTransfer = nullptr;
-		switch (Args.RHI)
-		{
-		case ERHI::D3D11:
-			TextureTransfer = new Private::FD3D11TextureTransfer();
-			break;
-		case ERHI::D3D12:
-			TextureTransfer = new Private::FD3D12TextureTransfer();
-			break;
-		default:
-			return nullptr;
-		}
-
-		if ((Private::FTextureTransferBase*)(TextureTransfer)->Initialize(Args))
-		{
-			TextureTransfersWrapper.Transfers[RHIIndex] = TextureTransfer;
-		}
-		else
-		{
-			delete TextureTransfer;
-		}
-
-		return TextureTransfersWrapper.Transfers[RHIIndex];
-#else
-		return nullptr;
-#endif
-	}
-
-	void CleanupTextureTransfer(ITextureTransfer* TextureTransfer)
-	{
-		TextureTransfersWrapper.CleanupTextureTransfer(TextureTransfer);
-	}
 }
 
 
