Remove use of RHIGetNativeDevice() / RHIGetNativeCommandQueue() in GPUTextureTransfer module
- Replaced with the per-RHI interfaces

Also removed FTextureTransfersWrapper
- Unused code

diff --git a/Engine/Plugins/Media/MediaIOFramework/Source/GPUTextureTransfer/GPUTextureTransfer.Build.cs b/Engine/Plugins/Media/MediaIOFramework/Source/GPUTextureTransfer/GPUTextureTransfer.Build.cs
--- a/Engine/Plugins/Media/MediaIOFramework/Source/GPUTextureTransfer/GPUTextureTransfer.Build.cs
+++ b/Engine/Plugins/Media/MediaIOFramework/Source/GPUTextureTransfer/GPUTextureTransfer.Build.cs
@@ -14,8 +14,8 @@ namespace UnrealBuildTool.Rules
 
 				if (Target.bCompileAgainstEngine)
 				{
-					AddEngineThirdPartyPrivateStaticDependencies(Target, "Vulkan");
-					PrivateDependencyModuleNames.Add("VulkanRHI");
+					AddEngineThirdPartyPrivateStaticDependencies(Target, "Vulkan", "DX11", "DX12");
+					PrivateDependencyModuleNames.AddRange(["VulkanRHI", "D3D11RHI", "D3D12RHI"]);
 				}
 				else
 				{
diff --git a/Engine/Plugins/Media/MediaIOFramework/Source/GPUTextureTransfer/Private/D3D11TextureTransfer.cpp b/Engine/Plugins/Media/MediaIOFramework/Source/GPUTextureTransfer/Private/D3D11TextureTransfer.cpp
--- a/Engine/Plugins/Media/MediaIOFramework/Source/GPUTextureTransfer/Private/D3D11TextureTransfer.cpp
+++ b/Engine/Plugins/Media/MediaIOFramework/Source/GPUTextureTransfer/Private/D3D11TextureTransfer.cpp
@@ -4,16 +4,16 @@
 
 #if DVP_SUPPORTED_PLATFORM
 #include "Windows/AllowWindowsPlatformTypes.h"
-
-#include "dvpapi_d3d11.h"
-
+	#include "dvpapi_d3d11.h"
 #include "Windows/HideWindowsPlatformTypes.h"
 
+#include "ID3D11DynamicRHI.h"
+
 namespace UE::GPUTextureTransfer::Private
 {
-DVPStatus FD3D11TextureTransfer::Init_Impl(const FInitializeDMAArgs& InArgs)
+DVPStatus FD3D11TextureTransfer::Init_Impl()
 {
-	D3DDevice = (ID3D11Device*)InArgs.RHIDevice;
+	D3DDevice = GetID3D11DynamicRHI()->RHIGetDevice();
 	return dvpInitD3D11Device(D3DDevice, 0);
 }
 
diff --git a/Engine/Plugins/Media/MediaIOFramework/Source/GPUTextureTransfer/Private/D3D11TextureTransfer.h b/Engine/Plugins/Media/MediaIOFramework/Source/GPUTextureTransfer/Private/D3D11TextureTransfer.h
--- a/Engine/Plugins/Media/MediaIOFramework/Source/GPUTextureTransfer/Private/D3D11TextureTransfer.h
+++ b/Engine/Plugins/Media/MediaIOFramework/Source/GPUTextureTransfer/Private/D3D11TextureTransfer.h
@@ -12,7 +12,7 @@ namespace UE::GPUTextureTransfer::Private
 class FD3D11TextureTransfer : public FTextureTransferBase
 {
 protected:
-	virtual DVPStatus Init_Impl(const FInitializeDMAArgs& InArgs) override;
+	virtual DVPStatus Init_Impl() override;
 	virtual DVPStatus GetConstants_Impl(uint32* OutBufferAddrAlignment, uint32* OutBufferGPUStrideAlignment, uint32* OutSemaphoreAddrAlignment, uint32* OutSemaphoreAllocSize, uint32* OutSemaphorePayloadOffset, uint32* OutSemaphorePayloadSize) const override;
 	virtual DVPStatus BindBuffer_Impl(DVPBufferHandle InBufferHandle) const override;
 	virtual DVPStatus CreateGPUResource_Impl(const FRegisterDMATextureArgs& InArgs, FTextureInfo* OutTextureInfo) const override;
diff --git a/Engine/Plugins/Media/MediaIOFramework/Source/GPUTextureTransfer/Private/D3D12TextureTransfer.cpp b/Engine/Plugins/Media/MediaIOFramework/Source/GPUTextureTransfer/Private/D3D12TextureTransfer.cpp
--- a/Engine/Plugins/Media/MediaIOFramework/Source/GPUTextureTransfer/Private/D3D12TextureTransfer.cpp
+++ b/Engine/Plugins/Media/MediaIOFramework/Source/GPUTextureTransfer/Private/D3D12TextureTransfer.cpp
@@ -4,18 +4,20 @@
 
 #if DVP_SUPPORTED_PLATFORM
 #include "Windows/AllowWindowsPlatformTypes.h"
-
-#include "dvpapi_d3d12.h"
-
+	#include "dvpapi_d3d12.h"
 #include "Windows/HideWindowsPlatformTypes.h"
 
+#include "ID3D12DynamicRHI.h"
 
 namespace UE::GPUTextureTransfer::Private
 {
-	DVPStatus FD3D12TextureTransfer::Init_Impl(const FInitializeDMAArgs& InArgs)
+	DVPStatus FD3D12TextureTransfer::Init_Impl()
 	{
-		D3DDevice = (ID3D12Device*)InArgs.RHIDevice;
-		D3DCommandQueue = (ID3D12CommandQueue*)InArgs.RHICommandQueue;
+		ID3D12DynamicRHI* D3D12RHI = GetID3D12DynamicRHI();
+
+		D3DDevice = D3D12RHI->RHIGetDevice(0);
+		D3DCommandQueue = D3D12RHI->RHIGetCommandQueue();
+
 		return dvpInitD3D12Device(D3DDevice, 0);
 	}
 
diff --git a/Engine/Plugins/Media/MediaIOFramework/Source/GPUTextureTransfer/Private/D3D12TextureTransfer.h b/Engine/Plugins/Media/MediaIOFramework/Source/GPUTextureTransfer/Private/D3D12TextureTransfer.h
--- a/Engine/Plugins/Media/MediaIOFramework/Source/GPUTextureTransfer/Private/D3D12TextureTransfer.h
+++ b/Engine/Plugins/Media/MediaIOFramework/Source/GPUTextureTransfer/Private/D3D12TextureTransfer.h
@@ -14,7 +14,7 @@ namespace UE::GPUTextureTransfer::Private
 class FD3D12TextureTransfer : public FTextureTransferBase
 {
 protected:
-	virtual DVPStatus Init_Impl(const FInitializeDMAArgs& InArgs) override;
+	virtual DVPStatus Init_Impl() override;
 	virtual DVPStatus GetConstants_Impl(uint32* OutBufferAddrAlignment, uint32* OutBufferGPUStrideAlignment, uint32* OutSemaphoreAddrAlignment, uint32* OutSemaphoreAllocSize, uint32* OutSemaphorePayloadOffset, uint32* OutSemaphorePayloadSize) const override;
 	virtual DVPStatus BindBuffer_Impl(DVPBufferHandle InBufferHandle) const override;
 	virtual DVPStatus CreateGPUResource_Impl(const FRegisterDMATextureArgs& InArgs, FTextureInfo* OutTextureInfo) const override;
diff --git a/Engine/Plugins/Media/MediaIOFramework/Source/GPUTextureTransfer/Private/GPUTextureTransferModule.cpp b/Engine/Plugins/Media/MediaIOFramework/Source/GPUTextureTransfer/Private/GPUTextureTransferModule.cpp
--- a/Engine/Plugins/Media/MediaIOFramework/Source/GPUTextureTransfer/Private/GPUTextureTransferModule.cpp
+++ b/Engine/Plugins/Media/MediaIOFramework/Source/GPUTextureTransfer/Private/GPUTextureTransferModule.cpp
@@ -252,22 +252,9 @@ void FGPUTextureTransferModule::InitializeTextureTransfer()
 			break;
 		}
 
-		UE::GPUTextureTransfer::FInitializeDMAArgs InitializeArgs;
-		InitializeArgs.RHI = RHI;
-		InitializeArgs.RHIDevice = GDynamicRHI->RHIGetNativeDevice();
-		InitializeArgs.RHICommandQueue = GDynamicRHI->RHIGetNativeGraphicsQueue();
-#if VULKAN_PLATFORM
-		if (RHI == UE::GPUTextureTransfer::ERHI::Vulkan)
-		{
-			IVulkanDynamicRHI* DynRHI = GetIVulkanDynamicRHI();
-			InitializeArgs.VulkanInstance = DynRHI->RHIGetVkInstance();
-			FMemory::Memcpy(InitializeArgs.RHIDeviceUUID, DynRHI->RHIGetVulkanDeviceUUID(), 16);
-		}
-#endif
-
 		const uint8 RHIIndex = static_cast<uint8>(RHI);
 		UE_LOG(LogGPUTextureTransfer, Display, TEXT("Initializing GPU Texture transfer"));
-		if (TextureTransfer->Initialize(InitializeArgs))
+		if (TextureTransfer->Initialize())
 		{
 			TransferObjects[RHIIndex] = TextureTransfer;
 		}
@@ -298,56 +285,4 @@ void FGPUTextureTransferModule::OnEnableGPUDirectCVarChange(IConsoleVariable* Co
 	}
 }
 
-namespace UE::GPUTextureTransfer
-{
-	struct FTextureTransfersWrapper
-	{
-		static constexpr uint8_t RHI_MAX = static_cast<uint8_t>(ERHI::RHI_MAX);
-		TArray<ITextureTransfer*> Transfers;
-
-		FTextureTransfersWrapper()
-		{
-			Transfers.SetNumUninitialized(RHI_MAX);
-
-			for (uint8_t RhiIt = 0; RhiIt < RHI_MAX; RhiIt++)
-			{
-				Transfers[RhiIt] = nullptr;
-			}
-		}
-
-		~FTextureTransfersWrapper()
-		{
-			// 0 is Invalid RHI
-			for (uint8_t RhiIt = 1; RhiIt < RHI_MAX; RhiIt++)
-			{
-				ITextureTransfer* TextureTransfer = Transfers[RhiIt];
-				if (TextureTransfer)
-				{
-					TextureTransfer->Uninitialize();
-					delete TextureTransfer;
-					Transfers[RhiIt] = nullptr;
-				}
-			}
-		}
-
-		void CleanupTextureTransfer(ITextureTransfer* TextureTransfer)
-		{
-			if (TextureTransfer)
-			{
-				for (uint8_t RhiIt = 1; RhiIt < RHI_MAX; RhiIt++)
-				{
-					if (Transfers[RhiIt] && Transfers[RhiIt] == TextureTransfer)
-					{
-						Transfers[RhiIt]->Uninitialize();
-						delete Transfers[RhiIt];
-						Transfers[RhiIt] = nullptr;
-					}
-				}
-			}
-		}
-
-	} TextureTransfersWrapper;
-}
-
-
 IMPLEMENT_MODULE(FGPUTextureTransferModule, GPUTextureTransfer);
diff --git a/Engine/Plugins/Media/MediaIOFramework/Source/GPUTextureTransfer/Private/TextureTransferBase.cpp b/Engine/Plugins/Media/MediaIOFramework/Source/GPUTextureTransfer/Private/TextureTransferBase.cpp
--- a/Engine/Plugins/Media/MediaIOFramework/Source/GPUTextureTransfer/Private/TextureTransferBase.cpp
+++ b/Engine/Plugins/Media/MediaIOFramework/Source/GPUTextureTransfer/Private/TextureTransferBase.cpp
@@ -60,7 +60,7 @@ namespace UE::GPUTextureTransfer::Private
 		*(volatile uint32*)(Semaphore) = Value;
 	}
 
-	bool FTextureTransferBase::Initialize(const FInitializeDMAArgs& Args)
+	bool FTextureTransferBase::Initialize()
 	{
 		FScopeLock ScopeLock{ &CriticalSection };
 
@@ -69,7 +69,7 @@ namespace UE::GPUTextureTransfer::Private
 			return false;
 		}
 
-		if (Init_Impl(Args) != DVP_STATUS_OK)
+		if (Init_Impl() != DVP_STATUS_OK)
 		{
 			UE_LOG(LogGPUTextureTransfer, Display, TEXT("GPU Direct failed to initialize."));
 			return false;
diff --git a/Engine/Plugins/Media/MediaIOFramework/Source/GPUTextureTransfer/Private/TextureTransferBase.h b/Engine/Plugins/Media/MediaIOFramework/Source/GPUTextureTransfer/Private/TextureTransferBase.h
--- a/Engine/Plugins/Media/MediaIOFramework/Source/GPUTextureTransfer/Private/TextureTransferBase.h
+++ b/Engine/Plugins/Media/MediaIOFramework/Source/GPUTextureTransfer/Private/TextureTransferBase.h
@@ -35,7 +35,7 @@ namespace UE::GPUTextureTransfer::Private
 		virtual void UnlockTexture(FRHITexture* InRHITexture) override;
 		// End ITextureTransfer interface
 
-		virtual bool Initialize(const FInitializeDMAArgs& Args) override;
+		virtual bool Initialize() override;
 		virtual bool Uninitialize() override;
 
 	protected:
@@ -78,7 +78,7 @@ namespace UE::GPUTextureTransfer::Private
 		};
 
 	protected:
-		virtual DVPStatus Init_Impl(const FInitializeDMAArgs& InArgs) = 0;
+		virtual DVPStatus Init_Impl() = 0;
 		virtual DVPStatus GetConstants_Impl(uint32* OutBufferAddrAlignment, uint32* OutBufferGPUStrideAlignment, uint32* OutSemaphoreAddrAlignment, uint32* OutSemaphoreAllocSize, uint32* OutSemaphorePayloadOffset, uint32* OutSemaphorePayloadSize) const = 0;
 		virtual DVPStatus BindBuffer_Impl(DVPBufferHandle InBufferHandle) const = 0;
 		virtual DVPStatus CreateGPUResource_Impl(const FRegisterDMATextureArgs& InArgs, FTextureInfo* OutTextureInfo) const = 0;
diff --git a/Engine/Plugins/Media/MediaIOFramework/Source/GPUTextureTransfer/Private/VulkanTextureTransfer.cpp b/Engine/Plugins/Media/MediaIOFramework/Source/GPUTextureTransfer/Private/VulkanTextureTransfer.cpp
--- a/Engine/Plugins/Media/MediaIOFramework/Source/GPUTextureTransfer/Private/VulkanTextureTransfer.cpp
+++ b/Engine/Plugins/Media/MediaIOFramework/Source/GPUTextureTransfer/Private/VulkanTextureTransfer.cpp
@@ -12,11 +12,17 @@
 
 namespace UE::GPUTextureTransfer::Private
 {
-	DVPStatus FVulkanTextureTransfer::Init_Impl(const FInitializeDMAArgs& InArgs)
+	DVPStatus FVulkanTextureTransfer::Init_Impl()
 	{
-		VulkanDevice = (VkDevice)InArgs.RHIDevice;
-		VulkanQueue = (VkQueue)InArgs.RHICommandQueue;
-		return dvpInitVkDevice(VulkanDevice, 0, (void*)InArgs.RHIDeviceUUID);
+		IVulkanDynamicRHI* VulkanRHI = GetIVulkanDynamicRHI();
+
+		VulkanDevice = VulkanRHI->RHIGetVkDevice();
+		VulkanQueue = VulkanRHI->RHIGetGraphicsVkQueue();
+
+		uint8 RHIDeviceUUID[16] {};
+		FMemory::Memcpy(RHIDeviceUUID, VulkanRHI->RHIGetVulkanDeviceUUID(), 16);
+
+		return dvpInitVkDevice(VulkanDevice, 0, RHIDeviceUUID);
 	}
 
 	DVPStatus FVulkanTextureTransfer::GetConstants_Impl(uint32* OutBufferAddrAlignment, uint32* OutBufferGPUStrideAlignment, uint32* OutSemaphoreAddrAlignment, uint32* OutSemaphoreAllocSize, uint32* OutSemaphorePayloadOffset, uint32* OutSemaphorePayloadSize) const
diff --git a/Engine/Plugins/Media/MediaIOFramework/Source/GPUTextureTransfer/Private/VulkanTextureTransfer.h b/Engine/Plugins/Media/MediaIOFramework/Source/GPUTextureTransfer/Private/VulkanTextureTransfer.h
--- a/Engine/Plugins/Media/MediaIOFramework/Source/GPUTextureTransfer/Private/VulkanTextureTransfer.h
+++ b/Engine/Plugins/Media/MediaIOFramework/Source/GPUTextureTransfer/Private/VulkanTextureTransfer.h
@@ -12,7 +12,7 @@ namespace UE::GPUTextureTransfer::Private
 class FVulkanTextureTransfer : public FTextureTransferBase
 {
 protected:
-	virtual DVPStatus Init_Impl(const FInitializeDMAArgs& InArgs) override;
+	virtual DVPStatus Init_Impl() override;
 	virtual DVPStatus GetConstants_Impl(uint32* OutBufferAddrAlignment, uint32* OutBufferGPUStrideAlignment, uint32* OutSemaphoreAddrAlignment, uint32* OutSemaphoreAllocSize, uint32* OutSemaphorePayloadOffset, uint32* OutSemaphorePayloadSize) const override;
 	virtual DVPStatus BindBuffer_Impl(DVPBufferHandle InBufferHandle) const override;
 	virtual DVPStatus CreateGPUResource_Impl(const FRegisterDMATextureArgs& InArgs, FTextureInfo* OutTextureInfo) const override;
diff --git a/Engine/Plugins/Media/MediaIOFramework/Source/GPUTextureTransfer/Public/GPUTextureTransfer.h b/Engine/Plugins/Media/MediaIOFramework/Source/GPUTextureTransfer/Public/GPUTextureTransfer.h
--- a/Engine/Plugins/Media/MediaIOFramework/Source/GPUTextureTransfer/Public/GPUTextureTransfer.h
+++ b/Engine/Plugins/Media/MediaIOFramework/Source/GPUTextureTransfer/Public/GPUTextureTransfer.h
@@ -19,22 +19,6 @@ namespace UE::GPUTextureTransfer
 
 	using FLoggingCallbackPtr = void(*)(const TCHAR* Format, ...);
 
-	struct FInitializeDMAArgs
-	{
-		/** Which RHI is being used. */
-		ERHI RHI = ERHI::Invalid;
-		/** The RHI device to use for gpu texture transfers. */
-		void* RHIDevice = nullptr;
-		/** The RHI device's command queue. */
-		void* RHICommandQueue = nullptr;
-
-		/** Vulkan instance (Vulkan only) */
-		void* VulkanInstance = nullptr;
-		/** Unique device identifier (Vulkan only) */
-		uint8 RHIDeviceUUID[16] = { 0 };
-		// End Vulkan Only
-	};
-
 	enum class EPixelFormat : uint8_t
 	{
 		PF_8Bit,
@@ -123,10 +107,7 @@ namespace UE::GPUTextureTransfer
 		 */
 		virtual uint32 GetTextureStride() const = 0;
 
-		virtual bool Initialize(const FInitializeDMAArgs& Args) = 0;
+		virtual bool Initialize() = 0;
 		virtual bool Uninitialize() = 0;
 	};
-
-	GPUTEXTURETRANSFER_API ITextureTransfer* GetTextureTransfer(const FInitializeDMAArgs& Args);
-	GPUTEXTURETRANSFER_API void CleanupTextureTransfer(ITextureTransfer* TextureTransfer);
 }
