MediaIO: Prevent custom color space selection causing crashes.

Note that this code will need to be replaced and unified with FMediaSourceColorSettings.

diff --git a/Engine/Plugins/Media/MediaIOFramework/Source/MediaIOCore/Public/CaptureCardMediaSource.h b/Engine/Plugins/Media/MediaIOFramework/Source/MediaIOCore/Public/CaptureCardMediaSource.h
--- a/Engine/Plugins/Media/MediaIOFramework/Source/MediaIOCore/Public/CaptureCardMediaSource.h
+++ b/Engine/Plugins/Media/MediaIOFramework/Source/MediaIOCore/Public/CaptureCardMediaSource.h
@@ -97,7 +97,7 @@ public:
 	bool bOverrideSourceColorSpace = true;
 
 	/** Color space of the source texture. */
-	UPROPERTY(BlueprintReadOnly, EditAnywhere, Category = "Video", meta = (EditCondition = "bOverrideSourceColorSpace"))
+	UPROPERTY(BlueprintReadOnly, EditAnywhere, Category = "Video", meta = (EditCondition = "bOverrideSourceColorSpace", InvalidEnumValues = "ETextureColorSpace::TCS_Custom"))
 	ETextureColorSpace OverrideSourceColorSpace = ETextureColorSpace::TCS_sRGB;
 
 	/**
