Media Source: Color management unification between img media and capture card sources.

* ImgMediaSourceColorSettings.h is deprecated & moved to BaseMediaSourceColorSettings.h
* EMediaSourceEncoding, FMediaSourceColorSettings & FNativeMediaSourceColorSettings are now shared between img media and media I/O.
* Native color settings shared pointers are thread-safe, and can be adjusted live without re-opening players.
* OpenColorIO transform settings renamed to OpenColorIO Override. Disables Encoding & Color Space UI if a valid transform is used.

diff --git a/Engine/Plugins/Media/AjaMedia/Source/AjaMedia/Private/Player/AjaMediaPlayer.cpp b/Engine/Plugins/Media/AjaMedia/Source/AjaMedia/Private/Player/AjaMediaPlayer.cpp
--- a/Engine/Plugins/Media/AjaMedia/Source/AjaMedia/Private/Player/AjaMediaPlayer.cpp
+++ b/Engine/Plugins/Media/AjaMedia/Source/AjaMedia/Private/Player/AjaMediaPlayer.cpp
@@ -63,7 +63,7 @@ FAjaMediaPlayer::FAjaMediaPlayer(IMediaEventSink& InEventSink)
 	, MetadataSamplePool(MakeUnique<FAjaMediaBinarySamplePool>())
 	, TextureSamplePool(MakeUnique<FAjaMediaTextureSamplePool>())
 	, EventSink(InEventSink)
-	, SupportedSampleTypes(EMediaIOSampleType::None)
+	, AjaColorSettings(MakeShared<FNativeMediaSourceColorSettings>())
 	, DeviceProvider(MakePimpl<FAjaDeviceProvider>())
 {
 }
@@ -613,12 +613,13 @@ bool FAjaMediaPlayer::OnInputFrameReceived(const AJA::AJAInputFrameData& InInput
 			bAjaWriteOutputRawDataCmdEnable = false;
 		}
 
-		UE::MediaIOCore::FColorFormatArgs ColorFormat;
 		FAjaMediaHDROptions HDROptions = UE::AjaMedia::MakeAjaMediaHDROptions(InVideoFrame.HDROptions);
-		
-		if (bOverrideSourceEncoding)
+		ensure(AjaColorSettings.IsValid());
+
+		const UE::Color::EEncoding SourceEncoding = SourceColorSettings.IsValid() ? SourceColorSettings->GetEncodingOverride() : UE::Color::EEncoding::None;
+		if (SourceEncoding != UE::Color::EEncoding::None)
 		{
-			ColorFormat.Encoding = (UE::Color::EEncoding) OverrideSourceEncoding;
+			AjaColorSettings->SetEncodingOverride(SourceEncoding);
 		}
 		else
 		{
@@ -627,45 +628,42 @@ bool FAjaMediaPlayer::OnInputFrameReceived(const AJA::AJAInputFrameData& InInput
 			switch (HDROptions.EOTF)
 			{
 			case EAjaHDRMetadataEOTF::SDR:
-				ColorFormat.Encoding = UE::Color::EEncoding::sRGB; // Missing support for BT.1886 so we default to sRGB instead.
+				AjaColorSettings->SetEncodingOverride(UE::Color::EEncoding::sRGB); // Missing support for BT.1886 so we default to sRGB instead.
 				break;
 			case EAjaHDRMetadataEOTF::PQ:
-				ColorFormat.Encoding = UE::Color::EEncoding::ST2084;
+				AjaColorSettings->SetEncodingOverride(UE::Color::EEncoding::ST2084);
 				break;
 			case EAjaHDRMetadataEOTF::HLG:
 				UE_MEDIA_IO_LOG_THROTTLE(LogAjaMedia, Warning, TimeBetweenLogs, TEXT("HLG EOTF is not supported at the moment. Defaulting to sRGB."));
-				ColorFormat.Encoding = UE::Color::EEncoding::sRGB;
+				AjaColorSettings->SetEncodingOverride(UE::Color::EEncoding::sRGB);
 				break;
 			default:
 				checkNoEntry();
-				ColorFormat.Encoding = UE::Color::EEncoding::None;			
 			}
 		}
-		
-		if (bOverrideSourceColorSpace)
+
+		if (SourceColorSettings.IsValid() && SourceColorSettings->HasColorSpaceOverride())
 		{
-			ColorFormat.ColorSpaceType = (UE::Color::EColorSpace) OverrideSourceColorSpace;
+			AjaColorSettings->SetColorSpaceOverride(SourceColorSettings->GetColorSpaceOverride(UE::Color::FColorSpace::GetSRGB()));
 		}
 		else
 		{
 			switch (HDROptions.Gamut)
 			{
 			case EAjaHDRMetadataGamut::Rec709:
-				ColorFormat.ColorSpaceType = UE::Color::EColorSpace::sRGB;
+				AjaColorSettings->SetColorSpaceOverride(UE::Color::FColorSpace::GetSRGB());
 				break;
 			case EAjaHDRMetadataGamut::Rec2020:
-				ColorFormat.ColorSpaceType = UE::Color::EColorSpace::Rec2020;
+				AjaColorSettings->SetColorSpaceOverride(UE::Color::FColorSpace::GetRec2020());
 				break;
 			default:
 				checkNoEntry();
-				ColorFormat.ColorSpaceType = UE::Color::EColorSpace::sRGB;
-				break;
 			}
 		}
 
 		if (AjaThreadCurrentTextureSample.IsValid())
 		{
-			if (AjaThreadCurrentTextureSample->SetProperties(InVideoFrame.Stride, InVideoFrame.Width, InVideoFrame.Height, VideoSampleFormat, DecodedTime, VideoFrameRate, DecodedTimecode, ColorFormat))
+			if (AjaThreadCurrentTextureSample->SetProperties(InVideoFrame.Stride, InVideoFrame.Width, InVideoFrame.Height, VideoSampleFormat, DecodedTime, VideoFrameRate, DecodedTimecode, AjaColorSettings))
 			{
 				AddVideoSample(AjaThreadCurrentTextureSample.ToSharedRef());
 			}
@@ -693,7 +691,7 @@ bool FAjaMediaPlayer::OnInputFrameReceived(const AJA::AJAInputFrameData& InInput
 				if (InVideoFrame.bIsProgressivePicture)
 				{
 					auto TextureSample = TextureSamplePool->AcquireShared();
-					if (TextureSample->InitializeProgressive(InVideoFrame, VideoSampleFormat, DecodedTime, VideoFrameRate, DecodedTimecode, ColorFormat))
+					if (TextureSample->InitializeProgressive(InVideoFrame, VideoSampleFormat, DecodedTime, VideoFrameRate, DecodedTimecode, AjaColorSettings))
 					{
 						TextureSample->SetColorConversionSettings(OCIOSettings);
 						AddVideoSample(TextureSample);
@@ -717,7 +715,7 @@ bool FAjaMediaPlayer::OnInputFrameReceived(const AJA::AJAInputFrameData& InInput
 						DecodedTime,
 						VideoFrameRate,
 						DecodedTimecode,
-						ColorFormat
+						AjaColorSettings
 					};
 
 					const TArray<TSharedRef<FMediaIOCoreTextureSampleBase>> DeinterlacedSamples = Deinterlace(FrameInfo);
diff --git a/Engine/Plugins/Media/AjaMedia/Source/AjaMedia/Private/Player/AjaMediaPlayer.h b/Engine/Plugins/Media/AjaMedia/Source/AjaMedia/Private/Player/AjaMediaPlayer.h
--- a/Engine/Plugins/Media/AjaMedia/Source/AjaMedia/Private/Player/AjaMediaPlayer.h
+++ b/Engine/Plugins/Media/AjaMedia/Source/AjaMedia/Private/Player/AjaMediaPlayer.h
@@ -182,7 +182,7 @@ private:
 	AJA::AJAInputChannel* InputChannel = nullptr;
 
 	/** Used to flag which sample types we advertise as supported for timed data monitoring */
-	EMediaIOSampleType SupportedSampleTypes;
+	EMediaIOSampleType SupportedSampleTypes = EMediaIOSampleType::None;
 
 	/** Previous frame timecode for stats purpose */
 	AJA::FTimecode AjaThreadPreviousFrameTimecode;
@@ -190,8 +190,12 @@ private:
 	/** Flag to indicate that pause is being requested */
 	std::atomic<bool> bPauseRequested = false;
 
+	/** AJA media source color format. */
 	EAjaMediaSourceColorFormat AjaColorFormat = EAjaMediaSourceColorFormat::YUV2_8bit;
 
+	/** Container for local media source color setting overrides. */
+	TSharedPtr<FNativeMediaSourceColorSettings, ESPMode::ThreadSafe> AjaColorSettings;
+
 	/** Device provider used to autodetect input format. */
 	TPimplPtr<class FAjaDeviceProvider> DeviceProvider;
 
diff --git a/Engine/Plugins/Media/AjaMedia/Source/AjaMedia/Private/Player/AjaMediaTextureSample.h b/Engine/Plugins/Media/AjaMedia/Source/AjaMedia/Private/Player/AjaMediaTextureSample.h
--- a/Engine/Plugins/Media/AjaMedia/Source/AjaMedia/Private/Player/AjaMediaTextureSample.h
+++ b/Engine/Plugins/Media/AjaMedia/Source/AjaMedia/Private/Player/AjaMediaTextureSample.h
@@ -23,9 +23,9 @@ public:
 	 * @param InTime The sample time (in the player's own clock).
 	 * @param InFrameRate The framerate of the media that produce the sample.
 	 * @param InTimecode The sample timecode if available.
-	 * @param bInIsSRGB Whether the sample is in sRGB space.
+	 * @param InSourceColorSettings Information about the texture color encoding and color space.
 	 */
-	bool InitializeProgressive(const AJA::AJAVideoFrameData& InVideoData, EMediaTextureSampleFormat InSampleFormat, FTimespan InTime, const FFrameRate& InFrameRate, const TOptional<FTimecode>& InTimecode, const UE::MediaIOCore::FColorFormatArgs& InColorFormatArgs)
+	bool InitializeProgressive(const AJA::AJAVideoFrameData& InVideoData, EMediaTextureSampleFormat InSampleFormat, FTimespan InTime, const FFrameRate& InFrameRate, const TOptional<FTimecode>& InTimecode, TSharedPtr<FNativeMediaSourceColorSettings, ESPMode::ThreadSafe> InSourceColorSettings)
 	{
 		return Super::Initialize(InVideoData.VideoBuffer
 			, InVideoData.VideoBufferSize
@@ -36,7 +36,7 @@ public:
 			, InTime
 			, InFrameRate
 			, InTimecode
-			, InColorFormatArgs);
+			, InSourceColorSettings);
 	}
 
 	/**
@@ -48,9 +48,9 @@ public:
 	 * @param InFrameRate The framerate of the media that produce the sample.
 	 * @param InTimecode The sample timecode if available.
 	 * @param bEven Only take the even frame from the image.
-	 * @param bInIsSRGB Whether the sample is in sRGB space.
+	 * @param InSourceColorSettings Information about the texture color encoding and color space.
 	 */
-	bool InitializeInterlaced_Halfed(const AJA::AJAVideoFrameData& InVideoData, EMediaTextureSampleFormat InSampleFormat, FTimespan InTime, const FFrameRate& InFrameRate, const TOptional<FTimecode>& InTimecode, bool bInEven, const UE::MediaIOCore::FColorFormatArgs& InColorFormatArgs)
+	bool InitializeInterlaced_Halfed(const AJA::AJAVideoFrameData& InVideoData, EMediaTextureSampleFormat InSampleFormat, FTimespan InTime, const FFrameRate& InFrameRate, const TOptional<FTimecode>& InTimecode, bool bInEven, TSharedPtr<FNativeMediaSourceColorSettings, ESPMode::ThreadSafe> InSourceColorSettings)
 	{
 		return Super::InitializeWithEvenOddLine(bInEven
 			, InVideoData.VideoBuffer
@@ -62,7 +62,7 @@ public:
 			, InTime
 			, InFrameRate
 			, InTimecode
-			, InColorFormatArgs);
+			, InSourceColorSettings);
 	}
 };
 
diff --git a/Engine/Plugins/Media/AjaMedia/Source/AjaMedia/Public/AjaMediaSource.h b/Engine/Plugins/Media/AjaMedia/Source/AjaMedia/Public/AjaMediaSource.h
--- a/Engine/Plugins/Media/AjaMedia/Source/AjaMedia/Public/AjaMediaSource.h
+++ b/Engine/Plugins/Media/AjaMedia/Source/AjaMedia/Public/AjaMediaSource.h
@@ -32,7 +32,7 @@ enum class EAjaMediaAudioChannel : uint8
 /**
  * Media source for AJA streams.
  */
-UCLASS(BlueprintType, hideCategories=(Platforms,Object), PrioritizeCategories=(AJA, Video, Audio, Ancillary, Synchronization, Debug), meta=(MediaIOCustomLayout="AJA"))
+UCLASS(BlueprintType, hideCategories=(Platforms,Object), PrioritizeCategories=(AJA, Video, ColorManagement, Audio, Ancillary, Synchronization, Debug), meta=(MediaIOCustomLayout="AJA"))
 class AJAMEDIA_API UAjaMediaSource : public UCaptureCardMediaSource
 {
 	GENERATED_BODY()
diff --git a/Engine/Plugins/Media/BlackmagicMedia/Source/BlackmagicMedia/Private/Player/BlackmagicMediaPlayer.cpp b/Engine/Plugins/Media/BlackmagicMedia/Source/BlackmagicMedia/Private/Player/BlackmagicMediaPlayer.cpp
--- a/Engine/Plugins/Media/BlackmagicMedia/Source/BlackmagicMedia/Private/Player/BlackmagicMediaPlayer.cpp
+++ b/Engine/Plugins/Media/BlackmagicMedia/Source/BlackmagicMedia/Private/Player/BlackmagicMediaPlayer.cpp
@@ -298,12 +298,13 @@ namespace BlackmagicMediaPlayerHelpers
 						bBlackmagicWriteOutputRawDataCmdEnable = false;
 					}
 					
-					UE::MediaIOCore::FColorFormatArgs ColorFormat;
 					FBlackmagicMediaHDROptions HDROptions = UE::BlackmagicMedia::MakeBlackmagicMediaHDROptions(InFrameInfo.HDRMetaData);
+					ensure(MediaPlayer->BlackmagicColorSettings.IsValid());
 		
-					if (MediaPlayer->bOverrideSourceEncoding)
+					const UE::Color::EEncoding SourceEncoding = MediaPlayer->SourceColorSettings.IsValid() ? MediaPlayer->SourceColorSettings->GetEncodingOverride() : UE::Color::EEncoding::None;
+					if (SourceEncoding != UE::Color::EEncoding::None)
 					{
-						ColorFormat.Encoding = (UE::Color::EEncoding) MediaPlayer->OverrideSourceEncoding;
+						MediaPlayer->BlackmagicColorSettings->SetEncodingOverride(SourceEncoding);
 					}
 					else
 					{
@@ -312,39 +313,36 @@ namespace BlackmagicMediaPlayerHelpers
 						switch (HDROptions.EOTF)
 						{
 						case EBlackmagicHDRMetadataEOTF::SDR:
-							ColorFormat.Encoding = UE::Color::EEncoding::sRGB; // Missing support for BT.1886 so we default to sRGB instead.
+							MediaPlayer->BlackmagicColorSettings->SetEncodingOverride(UE::Color::EEncoding::sRGB); // Missing support for BT.1886 so we default to sRGB instead.
 							break;
 						case EBlackmagicHDRMetadataEOTF::PQ:
-							ColorFormat.Encoding = UE::Color::EEncoding::ST2084;
+							MediaPlayer->BlackmagicColorSettings->SetEncodingOverride(UE::Color::EEncoding::ST2084);
 							break;
 						case EBlackmagicHDRMetadataEOTF::HLG:
 							UE_MEDIA_IO_LOG_THROTTLE(LogBlackmagicMedia, Warning, TimeBetweenLogs, TEXT("HLG EOTF is not supported at the moment. Defaulting to sRGB."));
-							ColorFormat.Encoding = UE::Color::EEncoding::sRGB;
+							MediaPlayer->BlackmagicColorSettings->SetEncodingOverride(UE::Color::EEncoding::sRGB);
 							break;
 						default:
 							checkNoEntry();
-							ColorFormat.Encoding = UE::Color::EEncoding::None;			
 						}
 					}
-		
-					if (MediaPlayer->bOverrideSourceColorSpace)
+
+					if (MediaPlayer->SourceColorSettings.IsValid() && MediaPlayer->SourceColorSettings->HasColorSpaceOverride())
 					{
-						ColorFormat.ColorSpaceType = (UE::Color::EColorSpace) MediaPlayer->OverrideSourceColorSpace;
+						MediaPlayer->BlackmagicColorSettings->SetColorSpaceOverride(MediaPlayer->SourceColorSettings->GetColorSpaceOverride(UE::Color::FColorSpace::GetSRGB()));
 					}
 					else
 					{
 						switch (HDROptions.Gamut)
 						{
 						case EBlackmagicHDRMetadataGamut::Rec709:
-							ColorFormat.ColorSpaceType = UE::Color::EColorSpace::sRGB;
+							MediaPlayer->BlackmagicColorSettings->SetColorSpaceOverride(UE::Color::FColorSpace::GetSRGB());
 							break;
 						case EBlackmagicHDRMetadataGamut::Rec2020:
-							ColorFormat.ColorSpaceType = UE::Color::EColorSpace::Rec2020;
+							MediaPlayer->BlackmagicColorSettings->SetColorSpaceOverride(UE::Color::FColorSpace::GetRec2020());
 							break;
 						default:
 							checkNoEntry();
-							ColorFormat.ColorSpaceType = UE::Color::EColorSpace::sRGB;
-							break;
 						}
 					}
 
@@ -372,7 +370,7 @@ namespace BlackmagicMediaPlayerHelpers
 						bool bInitializeResult = false;
 						if (bGPUDirectTexturesAvailable)
 						{
-							bInitializeResult = TextureSample->SetProperties(InFrameInfo.VideoPitch, InFrameInfo.VideoWidth, InFrameInfo.VideoHeight, SampleFormat, DecodedTime, MediaPlayer->VideoFrameRate, DecodedTimecode, ColorFormat);
+							bInitializeResult = TextureSample->SetProperties(InFrameInfo.VideoPitch, InFrameInfo.VideoWidth, InFrameInfo.VideoHeight, SampleFormat, DecodedTime, MediaPlayer->VideoFrameRate, DecodedTimecode, MediaPlayer->BlackmagicColorSettings);
 							if (bInitializeResult)
 							{
 								TextureSample->SetBuffer(InFrameInfo.VideoBuffer);
@@ -389,7 +387,7 @@ namespace BlackmagicMediaPlayerHelpers
 								, DecodedTime
 								, MediaPlayer->VideoFrameRate
 								, DecodedTimecode
-								, ColorFormat);
+								, MediaPlayer->BlackmagicColorSettings);
 						}
 
 						if (bInitializeResult)
@@ -419,7 +417,7 @@ namespace BlackmagicMediaPlayerHelpers
 							DecodedTime,
 							MediaPlayer->VideoFrameRate,
 							DecodedTimecode,
-							ColorFormat
+							MediaPlayer->BlackmagicColorSettings
 						};
 
 						const TArray<TSharedRef<FMediaIOCoreTextureSampleBase>> DeinterlacedSamples = MediaPlayer->Deinterlace(FrameInfo);
@@ -504,6 +502,7 @@ FBlackmagicMediaPlayer::FBlackmagicMediaPlayer(IMediaEventSink& InEventSink)
 	, AudioSamplePool(MakeUnique<FBlackmagicMediaAudioSamplePool>())
 	, TextureSamplePool(MakeUnique<FBlackmagicMediaTextureSamplePool>())
 	, SupportedSampleTypes(EMediaIOSampleType::None)
+	, BlackmagicColorSettings(MakeShared<FNativeMediaSourceColorSettings>())
 {
 }
 
diff --git a/Engine/Plugins/Media/BlackmagicMedia/Source/BlackmagicMedia/Private/Player/BlackmagicMediaPlayer.h b/Engine/Plugins/Media/BlackmagicMedia/Source/BlackmagicMedia/Private/Player/BlackmagicMediaPlayer.h
--- a/Engine/Plugins/Media/BlackmagicMedia/Source/BlackmagicMedia/Private/Player/BlackmagicMediaPlayer.h
+++ b/Engine/Plugins/Media/BlackmagicMedia/Source/BlackmagicMedia/Private/Player/BlackmagicMediaPlayer.h
@@ -158,5 +158,9 @@ private:
 	/** Used to flag which sample types we advertise as supported for timed data monitoring */
 	EMediaIOSampleType SupportedSampleTypes;
 
+	/** Blackmagic media source color format. */
 	EBlackmagicMediaSourceColorFormat BlackmagicColorFormat = EBlackmagicMediaSourceColorFormat::YUV8;
+
+	/** Container for local media source color setting overrides. */
+	TSharedPtr<FNativeMediaSourceColorSettings, ESPMode::ThreadSafe> BlackmagicColorSettings;
 };
diff --git a/Engine/Plugins/Media/BlackmagicMedia/Source/BlackmagicMedia/Public/BlackmagicMediaSource.h b/Engine/Plugins/Media/BlackmagicMedia/Source/BlackmagicMedia/Public/BlackmagicMediaSource.h
--- a/Engine/Plugins/Media/BlackmagicMedia/Source/BlackmagicMedia/Public/BlackmagicMediaSource.h
+++ b/Engine/Plugins/Media/BlackmagicMedia/Source/BlackmagicMedia/Public/BlackmagicMediaSource.h
@@ -33,7 +33,7 @@ enum class EBlackmagicMediaAudioChannel : uint8
 /**
  * Media source description for Blackmagic.
  */
-UCLASS(BlueprintType, hideCategories=(Platforms,Object), PrioritizeCategories=(Blackmagic, Video, Audio, Ancillary, Synchronization, Debug), meta=(MediaIOCustomLayout="Blackmagic"))
+UCLASS(BlueprintType, hideCategories=(Platforms,Object), PrioritizeCategories=(Blackmagic, Video, ColorManagement, Audio, Ancillary, Synchronization, Debug), meta=(MediaIOCustomLayout="Blackmagic"))
 class BLACKMAGICMEDIA_API UBlackmagicMediaSource : public UCaptureCardMediaSource
 {
 	GENERATED_BODY()
diff --git a/Engine/Plugins/Media/ImgMedia/Source/ImgMedia/Private/Assets/ImgMediaSourceColorSettings.cpp b/Engine/Plugins/Media/ImgMedia/Source/ImgMedia/Private/Assets/ImgMediaSourceColorSettings.cpp
deleted file mode 100644
diff --git a/Engine/Plugins/Media/ImgMedia/Source/ImgMedia/Private/Loader/ImgMediaLoader.cpp b/Engine/Plugins/Media/ImgMedia/Source/ImgMedia/Private/Loader/ImgMediaLoader.cpp
--- a/Engine/Plugins/Media/ImgMedia/Source/ImgMedia/Private/Loader/ImgMediaLoader.cpp
+++ b/Engine/Plugins/Media/ImgMedia/Source/ImgMedia/Private/Loader/ImgMediaLoader.cpp
@@ -23,7 +23,7 @@
 #include "Readers/IImgMediaReader.h"
 #include "Readers/DdsImgMediaReader.h"
 #include "ImgMediaLoaderWork.h"
-#include "ImgMediaSourceColorSettings.h"
+#include "BaseMediaSourceColorSettings.h"
 #include "Scheduler/ImgMediaScheduler.h"
 #include "Player/ImgMediaTextureSample.h"
 
diff --git a/Engine/Plugins/Media/ImgMedia/Source/ImgMedia/Private/Player/ImgMediaTextureSample.h b/Engine/Plugins/Media/ImgMedia/Source/ImgMedia/Private/Player/ImgMediaTextureSample.h
--- a/Engine/Plugins/Media/ImgMedia/Source/ImgMedia/Private/Player/ImgMediaTextureSample.h
+++ b/Engine/Plugins/Media/ImgMedia/Source/ImgMedia/Private/Player/ImgMediaTextureSample.h
@@ -8,7 +8,7 @@
 #include "Misc/Timespan.h"
 #include "Templates/SharedPointer.h"
 
-#include "ImgMediaSourceColorSettings.h"
+#include "BaseMediaSourceColorSettings.h"
 #include "Readers/IImgMediaReader.h"
 
 
@@ -177,7 +177,7 @@ public:
 		return IMediaTextureSample::GetEncodingType();
 	}
 
-	virtual UE::Color::EChromaticAdaptationMethod GetChromaticAdapationMethod() const
+	virtual UE::Color::EChromaticAdaptationMethod GetChromaticAdapationMethod() const override
 	{
 		if (SourceColorSettings.IsValid())
 		{
diff --git a/Engine/Plugins/Media/ImgMedia/Source/ImgMedia/Private/Readers/ExrImgMediaReaderGpu.cpp b/Engine/Plugins/Media/ImgMedia/Source/ImgMedia/Private/Readers/ExrImgMediaReaderGpu.cpp
--- a/Engine/Plugins/Media/ImgMedia/Source/ImgMedia/Private/Readers/ExrImgMediaReaderGpu.cpp
+++ b/Engine/Plugins/Media/ImgMedia/Source/ImgMedia/Private/Readers/ExrImgMediaReaderGpu.cpp
@@ -18,7 +18,7 @@
 #include "Engine/TextureRenderTarget2D.h"
 #include "GlobalShader.h"
 #include "CommonRenderResources.h"
-#include "ImgMediaSourceColorSettings.h"
+#include "BaseMediaSourceColorSettings.h"
 #include "Loader/ImgMediaLoader.h"
 #include "PostProcess/DrawRectangle.h"
 #include "RHIStaticStates.h"
diff --git a/Engine/Plugins/Media/ImgMedia/Source/ImgMedia/Public/ImgMediaSource.h b/Engine/Plugins/Media/ImgMedia/Source/ImgMedia/Public/ImgMediaSource.h
--- a/Engine/Plugins/Media/ImgMedia/Source/ImgMedia/Public/ImgMediaSource.h
+++ b/Engine/Plugins/Media/ImgMedia/Source/ImgMedia/Public/ImgMediaSource.h
@@ -10,8 +10,7 @@
 #include "UObject/ObjectMacros.h"
 #include "UObject/ScriptMacros.h"
 #include "BaseMediaSource.h"
-
-#include "ImgMediaSourceColorSettings.h"
+#include "BaseMediaSourceColorSettings.h"
 
 #include "ImgMediaSource.generated.h"
 
diff --git a/Engine/Plugins/Media/ImgMedia/Source/ImgMedia/Public/ImgMediaSourceColorSettings.h b/Engine/Plugins/Media/ImgMedia/Source/ImgMedia/Public/ImgMediaSourceColorSettings.h
--- a/Engine/Plugins/Media/ImgMedia/Source/ImgMedia/Public/ImgMediaSourceColorSettings.h
+++ b/Engine/Plugins/Media/ImgMedia/Source/ImgMedia/Public/ImgMediaSourceColorSettings.h
@@ -2,106 +2,6 @@
 
 #pragma once
 
-#include "CoreTypes.h"
-#include "ColorManagement/ColorSpace.h"
-#include "Engine/TextureDefines.h"
-#include "IMediaOptions.h"
+UE_DEPRECATED_HEADER(5.8, "Use BaseMediaSourceColorSettings.h from the MediaAssets module instead of ImgMediaSourceColorSettings.h.")
 
-#include "ImgMediaSourceColorSettings.generated.h"
-
-// TODO: EMediaSourceEncoding, FMediaSourceColorSettings & FNativeMediaSourceColorSettings should be moved up to BaseMediaSource.h and unified with UCaptureCardMediaSource settings.
-
-/** List of source encodings that can be converted to linear. (Integer values match the ETextureSourceEncoding values in TextureDefines.h */
-UENUM()
-enum class EMediaSourceEncoding : uint8
-{
-	MSE_None	= 0 UMETA(DisplayName = "None", ToolTip = "The source encoding is not overridden."),
-	MSE_Linear	= 1 UMETA(DisplayName = "Linear", ToolTip = "The source encoding is considered linear."),
-	MSE_sRGB	= 2 UMETA(DisplayName = "sRGB", ToolTip = "sRGB source encoding to be linearized"),
-	MSE_ST2084	= 3 UMETA(DisplayName = "ST 2084/PQ", ToolTip = "SMPTE ST 2084/PQ source encoding to be linearized"),
-	MSE_SLog3	= 12 UMETA(DisplayName = "SLog3", ToolTip = "Sony SLog3 source encoding to be linearized"),
-
-	MSE_MAX,
-};
-
-/* Manual definition of media source color space & encoding. */
-USTRUCT(BlueprintType)
-struct FMediaSourceColorSettings
-{
-	GENERATED_USTRUCT_BODY()
-
-	FMediaSourceColorSettings()
-		: EncodingOverride(EMediaSourceEncoding::MSE_None)
-		, ColorSpaceOverride(ETextureColorSpace::TCS_None)
-		, RedChromaticityCoordinate(FVector2D::ZeroVector)
-		, GreenChromaticityCoordinate(FVector2D::ZeroVector)
-		, BlueChromaticityCoordinate(FVector2D::ZeroVector)
-		, WhiteChromaticityCoordinate(FVector2D::ZeroVector)
-		, ChromaticAdaptationMethod(static_cast<ETextureChromaticAdaptationMethod>(UE::Color::DEFAULT_CHROMATIC_ADAPTATION_METHOD))
-	{}
-
-	/** Source encoding of the media. */
-	UPROPERTY(EditAnywhere, BlueprintReadWrite, Category = ColorManagement)
-	EMediaSourceEncoding EncodingOverride;
-
-	/** Source color space of the media. */
-	UPROPERTY(EditAnywhere, BlueprintReadWrite, Category = ColorManagement)
-	ETextureColorSpace ColorSpaceOverride;
-
-	/** Red chromaticity coordinate of the source color space. */
-	UPROPERTY(EditAnywhere, BlueprintReadWrite, Category = ColorManagement, meta = (EditCondition = "ColorSpaceOverride == ETextureColorSpace::TCS_Custom"))
-	FVector2D RedChromaticityCoordinate;
-
-	/** Green chromaticity coordinate of the source color space. */
-	UPROPERTY(EditAnywhere, BlueprintReadWrite, Category = ColorManagement, meta = (EditCondition = "ColorSpaceOverride == ETextureColorSpace::TCS_Custom"))
-	FVector2D GreenChromaticityCoordinate;
-
-	/** Blue chromaticity coordinate of the source color space. */
-	UPROPERTY(EditAnywhere, BlueprintReadWrite, Category = ColorManagement, meta = (EditCondition = "ColorSpaceOverride == ETextureColorSpace::TCS_Custom"))
-	FVector2D BlueChromaticityCoordinate;
-
-	/** White chromaticity coordinate of the source color space. */
-	UPROPERTY(EditAnywhere, BlueprintReadWrite, Category = ColorManagement, meta = (EditCondition = "ColorSpaceOverride == ETextureColorSpace::TCS_Custom"))
-	FVector2D WhiteChromaticityCoordinate;
-
-	/** Chromatic adaption method applied if the source white point differs from the working color space white point. */
-	UPROPERTY(EditAnywhere, BlueprintReadWrite, Category = ColorManagement, meta = (EditCondition = "ColorSpaceOverride != ETextureColorSpace::TCS_None"))
-	ETextureChromaticAdaptationMethod ChromaticAdaptationMethod;
-
-#if WITH_EDITOR
-	/** Update the chromaticity coordinates member variables based on the color space choice (unless custom). */
-	IMGMEDIA_API void UpdateColorSpaceChromaticities();
-#endif
-};
-
-/* Engine-native color source settings container for media option. */
-struct FNativeMediaSourceColorSettings : public IMediaOptions::FDataContainer
-{
-	/** Constructor */
-	FNativeMediaSourceColorSettings();
-
-	/** Updates the native settings from user-controlled settings. */
-	void Update(const FMediaSourceColorSettings& InSettings);
-
-	/** Color encoding override getter. */
-	UE::Color::EEncoding GetEncodingOverride() const { return EncodingOverride; }
-
-	/** Color space override getter. */
-	const UE::Color::FColorSpace& GetColorSpaceOverride(const UE::Color::FColorSpace& InDefaultColorSpace) const;
-
-	/** Chromatic adaptation getter. */
-	UE::Color::EChromaticAdaptationMethod GetChromaticAdaptationMethod() const { return ChromaticAdaptationMethod; }
-
-private:
-	/** Manual source encoding override. */
-	std::atomic<UE::Color::EEncoding> EncodingOverride;
-
-	/** Manual source color space override. */
-	TOptional<UE::Color::FColorSpace> ColorSpaceOverride;
-
-	/** Chromatic adapation to be used on manual source color space override. */
-	std::atomic<UE::Color::EChromaticAdaptationMethod> ChromaticAdaptationMethod;
-
-	/** Protects color space override variable. */
-	mutable FCriticalSection ColorSpaceCriticalSection;
-};
+#include "BaseMediaSourceColorSettings.h"
diff --git a/Engine/Plugins/Media/MediaFrameworkUtilities/Source/MediaFrameworkUtilitiesEditor/Private/AssetEditor/MediaProfileEditor.cpp b/Engine/Plugins/Media/MediaFrameworkUtilities/Source/MediaFrameworkUtilitiesEditor/Private/AssetEditor/MediaProfileEditor.cpp
--- a/Engine/Plugins/Media/MediaFrameworkUtilities/Source/MediaFrameworkUtilitiesEditor/Private/AssetEditor/MediaProfileEditor.cpp
+++ b/Engine/Plugins/Media/MediaFrameworkUtilities/Source/MediaFrameworkUtilitiesEditor/Private/AssetEditor/MediaProfileEditor.cpp
@@ -12,6 +12,7 @@
 #include "SMediaProfileDetailsPanel.h"
 #include "SMediaProfileSourcesTreeView.h"
 #include "SMediaProfileViewport.h"
+#include "CaptureCardMediaSource.h"
 #include "CaptureTab/SMediaFrameworkCapture.h"
 #include "Framework/Application/SlateApplication.h"
 #include "GameFramework/WorldSettings.h"
@@ -544,12 +545,22 @@ void FMediaProfileEditor::OnObjectPropertyChanged(UObject* InObject, FPropertyCh
 			{
 				if (PlaybackManager->IsSourceOpenFromIndex(SrcIndex))
 				{
-					UMediaProfilePlaybackManager::FCloseSourceArgs Args;
-					Args.Consumer = this;
-					Args.bForceClose = true;
-		
-					PlaybackManager->CloseSourceFromIndex(SrcIndex, Args);
-					PlaybackManager->OpenSourceFromIndex(SrcIndex, this);
+					bool bPropertyChangeRequiresRestart = true;
+
+					if (InPropertyChangedEvent.GetMemberPropertyName() == GET_MEMBER_NAME_CHECKED(UCaptureCardMediaSource, SourceColorSettings))
+					{
+						bPropertyChangeRequiresRestart = false;
+					}
+
+					if (bPropertyChangeRequiresRestart)
+					{
+						UMediaProfilePlaybackManager::FCloseSourceArgs Args;
+						Args.Consumer = this;
+						Args.bForceClose = true;
+
+						PlaybackManager->CloseSourceFromIndex(SrcIndex, Args);
+						PlaybackManager->OpenSourceFromIndex(SrcIndex, this);
+					}
 				}
 			}
 		}
diff --git a/Engine/Plugins/Media/MediaIOFramework/Source/MediaIOCore/Private/CaptureCardMediaSource.cpp b/Engine/Plugins/Media/MediaIOFramework/Source/MediaIOCore/Private/CaptureCardMediaSource.cpp
--- a/Engine/Plugins/Media/MediaIOFramework/Source/MediaIOCore/Private/CaptureCardMediaSource.cpp
+++ b/Engine/Plugins/Media/MediaIOFramework/Source/MediaIOCore/Private/CaptureCardMediaSource.cpp
@@ -2,10 +2,13 @@
 
 #include "CaptureCardMediaSource.h"
 
+#include "UObject/FortniteMainBranchObjectVersion.h"
+
 #include UE_INLINE_GENERATED_CPP_BY_NAME(CaptureCardMediaSource)
 
 
 UCaptureCardMediaSource::UCaptureCardMediaSource()
+	: NativeSourceColorSettings(MakeShared<FNativeMediaSourceColorSettings, ESPMode::ThreadSafe>())
 {
 	Deinterlacer = CreateDefaultSubobject<UBobDeinterlacer>("Deinterlacer");
 }
@@ -17,16 +20,6 @@ int64 UCaptureCardMediaSource::GetMediaOption(const FName& Key, int64 DefaultVal
 		return static_cast<int64>(InterlaceFieldOrder);
 	}
 
-	if (Key == UE::CaptureCardMediaSource::SourceColorSpace)
-	{
-		return (int64) OverrideSourceColorSpace;
-	}
-
-	if (Key == UE::CaptureCardMediaSource::SourceEncoding)
-	{
-		return (int64) OverrideSourceEncoding;
-	}
-
 	if (Key == UE::CaptureCardMediaSource::EvaluationType)
 	{
 		return (int64) EvaluationType;
@@ -44,6 +37,12 @@ TSharedPtr<IMediaOptions::FDataContainer, ESPMode::ThreadSafe> UCaptureCardMedia
 		return MakeShared<FOpenColorIODataContainer>(MoveTemp(Container));
 	}
 
+	if (Key == UE::CaptureCardMediaSource::SourceColorSettingsOption)
+	{
+		NativeSourceColorSettings->Update(SourceColorSettings);
+		return NativeSourceColorSettings;
+	}
+
 	return Super::GetMediaOption(Key, DefaultValue);
 }
 
@@ -72,16 +71,6 @@ bool UCaptureCardMediaSource::GetMediaOption(const FName& Key, bool bDefaultValu
 		return bFramelock;
 	}
 
-	if (Key == UE::CaptureCardMediaSource::OverrideSourceColorSpace)
-	{
-		return bOverrideSourceColorSpace;
-	}
-
-	if (Key == UE::CaptureCardMediaSource::OverrideSourceEncoding)
-	{
-		return bOverrideSourceEncoding;
-	}
-
 	return UTimeSynchronizableMediaSource::GetMediaOption(Key, bDefaultValue);
 }
 
@@ -89,10 +78,7 @@ bool UCaptureCardMediaSource::HasMediaOption(const FName& Key) const
 {
 	if (Key == UE::CaptureCardMediaSource::Deinterlacer
 		|| Key == UE::CaptureCardMediaSource::InterlaceFieldOrder
-		|| Key == UE::CaptureCardMediaSource::OverrideSourceEncoding
-		|| Key == UE::CaptureCardMediaSource::OverrideSourceColorSpace
-		|| Key == UE::CaptureCardMediaSource::SourceEncoding
-		|| Key == UE::CaptureCardMediaSource::SourceColorSpace
+		|| Key == UE::CaptureCardMediaSource::SourceColorSettingsOption
 		|| Key == UE::CaptureCardMediaSource::RenderJIT
 		|| Key == UE::CaptureCardMediaSource::Framelock
 		|| Key == UE::CaptureCardMediaSource::EvaluationType
@@ -105,6 +91,30 @@ bool UCaptureCardMediaSource::HasMediaOption(const FName& Key) const
 }
 
 #if WITH_EDITOR
+void UCaptureCardMediaSource::PostLoad()
+{
+	Super::PostLoad();
+
+#if WITH_EDITORONLY_DATA
+PRAGMA_DISABLE_DEPRECATION_WARNINGS
+	if (GetLinkerCustomVersion(FFortniteMainBranchObjectVersion::GUID) < FFortniteMainBranchObjectVersion::MediaColorSettingsUnification)
+	{
+		if (bOverrideSourceEncoding_DEPRECATED)
+		{
+			SourceColorSettings.EncodingOverride = static_cast<EMediaSourceEncoding>(OverrideSourceEncoding_DEPRECATED);
+			bOverrideSourceEncoding_DEPRECATED = false;
+		}
+
+		if (bOverrideSourceColorSpace_DEPRECATED)
+		{
+			SourceColorSettings.ColorSpaceOverride = OverrideSourceColorSpace_DEPRECATED;
+			bOverrideSourceColorSpace_DEPRECATED = false;
+		}
+	}
+PRAGMA_ENABLE_DEPRECATION_WARNINGS
+#endif
+}
+
 bool UCaptureCardMediaSource::CanEditChange(const FProperty* InProperty) const
 {
 	if (!Super::CanEditChange(InProperty))
@@ -117,6 +127,12 @@ bool UCaptureCardMediaSource::CanEditChange(const FProperty* InProperty) const
 		return bRenderJIT && EvaluationType == EMediaIOSampleEvaluationType::Timecode && bUseTimeSynchronization;
 	}
 
+	if (InProperty->GetFName() == GET_MEMBER_NAME_CHECKED(UCaptureCardMediaSource, SourceColorSettings))
+	{
+		// OpenColorIO takes priority over the regular source color setting overrides.
+		return !ColorConversionSettings.IsValid();
+	}
+
 	return true;
 }
 
@@ -149,5 +165,16 @@ void UCaptureCardMediaSource::PostEditChangeChainProperty(struct FPropertyChange
 			bFramelock = false;
 		}
 	}
+
+	if (PropertyChangedEvent.GetPropertyName() == GET_MEMBER_NAME_CHECKED(FMediaSourceColorSettings, ColorSpaceOverride))
+	{
+		SourceColorSettings.UpdateColorSpaceChromaticities();
+	}
+
+	const FName StructName = PropertyChangedEvent.PropertyChain.GetActiveMemberNode()->GetValue()->GetFName();
+	if (StructName == GET_MEMBER_NAME_CHECKED(ThisClass, SourceColorSettings))
+	{
+		NativeSourceColorSettings->Update(SourceColorSettings);
+	}
 }
 #endif //WITH_EDITOR
diff --git a/Engine/Plugins/Media/MediaIOFramework/Source/MediaIOCore/Private/MediaIOCoreDeinterlacer.cpp b/Engine/Plugins/Media/MediaIOFramework/Source/MediaIOCore/Private/MediaIOCoreDeinterlacer.cpp
--- a/Engine/Plugins/Media/MediaIOFramework/Source/MediaIOCore/Private/MediaIOCoreDeinterlacer.cpp
+++ b/Engine/Plugins/Media/MediaIOFramework/Source/MediaIOCore/Private/MediaIOCoreDeinterlacer.cpp
@@ -24,7 +24,7 @@ namespace UE::MediaIOCore
 			, InVideoFrame.Time
 			, InVideoFrame.FrameRate
 			, InVideoFrame.Timecode
-			, InVideoFrame.ColorFormatArgs))
+			, InVideoFrame.SourceColorSettings))
 		{
 			Samples.Add(TextureSample.ToSharedRef());
 		}
@@ -72,7 +72,7 @@ namespace UE::MediaIOCore
 				, InVideoFrame.Time
 				, InVideoFrame.FrameRate
 				, InVideoFrame.Timecode
-				, InVideoFrame.ColorFormatArgs))
+				, InVideoFrame.SourceColorSettings))
 			{
 				Samples.Add(TextureSampleEven.ToSharedRef());
 			}
@@ -88,7 +88,7 @@ namespace UE::MediaIOCore
 				, InVideoFrame.Time
 				, InVideoFrame.FrameRate
 				, InVideoFrame.Timecode
-				, InVideoFrame.ColorFormatArgs))
+				, InVideoFrame.SourceColorSettings))
 			{
 				Samples.Add(TextureSampleOdd.ToSharedRef());
 			}
@@ -129,7 +129,7 @@ namespace UE::MediaIOCore
 				, InVideoFrame.Time
 				, InVideoFrame.FrameRate
 				, InVideoFrame.Timecode
-				, InVideoFrame.ColorFormatArgs))
+				, InVideoFrame.SourceColorSettings))
 			{
 				Samples.Add(TextureSample.ToSharedRef());
 			}
@@ -161,7 +161,7 @@ namespace UE::MediaIOCore
 				, InVideoFrame.Time
 				, InVideoFrame.FrameRate
 				, InVideoFrame.Timecode
-				, InVideoFrame.ColorFormatArgs))
+				, InVideoFrame.SourceColorSettings))
 			{
 				Samples.Add(TextureSample.ToSharedRef());
 			}
diff --git a/Engine/Plugins/Media/MediaIOFramework/Source/MediaIOCore/Private/Player/MediaIOCorePlayerBase.cpp b/Engine/Plugins/Media/MediaIOFramework/Source/MediaIOCore/Private/Player/MediaIOCorePlayerBase.cpp
--- a/Engine/Plugins/Media/MediaIOFramework/Source/MediaIOCore/Private/Player/MediaIOCorePlayerBase.cpp
+++ b/Engine/Plugins/Media/MediaIOFramework/Source/MediaIOCore/Private/Player/MediaIOCorePlayerBase.cpp
@@ -243,11 +243,6 @@ bool FMediaIOCorePlayerBase::Open(const FString& Url, const IMediaOptions* Optio
 		Deinterlacer = MakeShared<UE::MediaIOCore::FDeinterlacer>(UE::MediaIOCore::FDeinterlacer::FOnAcquireSample_AnyThread::CreateSP(this, &FMediaIOCorePlayerBase::AcquireTextureSample_AnyThread), InterlaceFieldOrder);
 	}
 
-	bOverrideSourceEncoding = Options->GetMediaOption(UE::CaptureCardMediaSource::OverrideSourceEncoding, true);
-	OverrideSourceEncoding = (ETextureSourceEncoding)Options->GetMediaOption(UE::CaptureCardMediaSource::SourceEncoding, (int64)ETextureSourceEncoding::TSE_Linear);
-	bOverrideSourceColorSpace = Options->GetMediaOption(UE::CaptureCardMediaSource::OverrideSourceColorSpace, true);
-	OverrideSourceColorSpace = (ETextureColorSpace)Options->GetMediaOption(UE::CaptureCardMediaSource::SourceColorSpace, (int64)ETextureColorSpace::TCS_None);
-
 	if (bReadMediaOptions)
 	{
 		CurrentState = EMediaState::Preparing;
@@ -576,12 +571,19 @@ bool FMediaIOCorePlayerBase::ReadMediaOptions(const IMediaOptions* Options)
 	BaseSettings.PlayerTimeOffset = MediaIOCorePlayerDetail::HighPerformanceClockDelta;
 	BaseSettings.AbsoluteMaxBufferSize = MediaIOCorePlayerDetail::CVarMediaIOMaxBufferSize.GetValueOnGameThread();
 
-	const TSharedPtr<UCaptureCardMediaSource::FOpenColorIODataContainer> OCIODataContainer = StaticCastSharedPtr<UCaptureCardMediaSource::FOpenColorIODataContainer>(Options->GetMediaOption(UE::CaptureCardMediaSource::OpenColorIOSettings, TSharedPtr<UCaptureCardMediaSource::FOpenColorIODataContainer, ESPMode::ThreadSafe>()));
+	TSharedPtr<IMediaOptions::FDataContainer, ESPMode::ThreadSafe> DefaultOption;
+	const TSharedPtr<UCaptureCardMediaSource::FOpenColorIODataContainer> OCIODataContainer = StaticCastSharedPtr<UCaptureCardMediaSource::FOpenColorIODataContainer>(Options->GetMediaOption(UE::CaptureCardMediaSource::OpenColorIOSettings, DefaultOption));
 	if (OCIODataContainer)
 	{
 		OCIOSettings = MakeShared<FOpenColorIOColorConversionSettings>(OCIODataContainer->ColorConversionSettings);
 	}
 
+	TSharedPtr<IMediaOptions::FDataContainer, ESPMode::ThreadSafe> DataContainer = Options->GetMediaOption(UE::CaptureCardMediaSource::SourceColorSettingsOption, DefaultOption);
+	if (DataContainer.IsValid())
+	{
+		SourceColorSettings = StaticCastSharedPtr<FNativeMediaSourceColorSettings, IMediaOptions::FDataContainer, ESPMode::ThreadSafe>(DataContainer);
+	}
+
 	return true;
 }
 
diff --git a/Engine/Plugins/Media/MediaIOFramework/Source/MediaIOCore/Private/Shared/MediaIOCoreTextureSampleBase.cpp b/Engine/Plugins/Media/MediaIOFramework/Source/MediaIOCore/Private/Shared/MediaIOCoreTextureSampleBase.cpp
--- a/Engine/Plugins/Media/MediaIOFramework/Source/MediaIOCore/Private/Shared/MediaIOCoreTextureSampleBase.cpp
+++ b/Engine/Plugins/Media/MediaIOFramework/Source/MediaIOCore/Private/Shared/MediaIOCoreTextureSampleBase.cpp
@@ -25,15 +25,16 @@ FMediaIOCoreTextureSampleBase::FMediaIOCoreTextureSampleBase()
 	, Width(0)
 	, Height(0)
 	, bIsAwaitingForGPUTransfer(false)
+	, SourceColorSettings()
 {
 }
 
 
-bool FMediaIOCoreTextureSampleBase::Initialize(const void* InVideoBuffer, uint32 InBufferSize, uint32 InStride, uint32 InWidth, uint32 InHeight, EMediaTextureSampleFormat InSampleFormat, FTimespan InTime, const FFrameRate& InFrameRate, const TOptional<FTimecode>& InTimecode, const UE::MediaIOCore::FColorFormatArgs& InColorFormatArgs)
+bool FMediaIOCoreTextureSampleBase::Initialize(const void* InVideoBuffer, uint32 InBufferSize, uint32 InStride, uint32 InWidth, uint32 InHeight, EMediaTextureSampleFormat InSampleFormat, FTimespan InTime, const FFrameRate& InFrameRate, const TOptional<FTimecode>& InTimecode, TSharedPtr<FNativeMediaSourceColorSettings, ESPMode::ThreadSafe> InSourceColorSettings)
 {
 	FreeSample();
 
-	if (!SetProperties(InStride, InWidth, InHeight, InSampleFormat, InTime, InFrameRate, InTimecode, InColorFormatArgs))
+	if (!SetProperties(InStride, InWidth, InHeight, InSampleFormat, InTime, InFrameRate, InTimecode, InSourceColorSettings))
 	{
 		return false;
 	}
@@ -42,11 +43,11 @@ bool FMediaIOCoreTextureSampleBase::Initialize(const void* InVideoBuffer, uint32
 }
 
 
-bool FMediaIOCoreTextureSampleBase::Initialize(const TArray<uint8>& InVideoBuffer, uint32 InStride, uint32 InWidth, uint32 InHeight, EMediaTextureSampleFormat InSampleFormat, FTimespan InTime, const FFrameRate& InFrameRate, const TOptional<FTimecode>& InTimecode, const UE::MediaIOCore::FColorFormatArgs& InColorFormatArgs)
+bool FMediaIOCoreTextureSampleBase::Initialize(const TArray<uint8>& InVideoBuffer, uint32 InStride, uint32 InWidth, uint32 InHeight, EMediaTextureSampleFormat InSampleFormat, FTimespan InTime, const FFrameRate& InFrameRate, const TOptional<FTimecode>& InTimecode, TSharedPtr<FNativeMediaSourceColorSettings, ESPMode::ThreadSafe> InSourceColorSettings)
 {
 	FreeSample();
 
-	if (!SetProperties(InStride, InWidth, InHeight, InSampleFormat, InTime, InFrameRate, InTimecode, InColorFormatArgs))
+	if (!SetProperties(InStride, InWidth, InHeight, InSampleFormat, InTime, InFrameRate, InTimecode, InSourceColorSettings))
 	{
 		return false;
 	}
@@ -55,11 +56,11 @@ bool FMediaIOCoreTextureSampleBase::Initialize(const TArray<uint8>& InVideoBuffe
 }
 
 
-bool FMediaIOCoreTextureSampleBase::Initialize(TArray<uint8>&& InVideoBuffer, uint32 InStride, uint32 InWidth, uint32 InHeight, EMediaTextureSampleFormat InSampleFormat, FTimespan InTime, const FFrameRate& InFrameRate, const TOptional<FTimecode>& InTimecode, const UE::MediaIOCore::FColorFormatArgs& InColorFormatArgs)
+bool FMediaIOCoreTextureSampleBase::Initialize(TArray<uint8>&& InVideoBuffer, uint32 InStride, uint32 InWidth, uint32 InHeight, EMediaTextureSampleFormat InSampleFormat, FTimespan InTime, const FFrameRate& InFrameRate, const TOptional<FTimecode>& InTimecode, TSharedPtr<FNativeMediaSourceColorSettings, ESPMode::ThreadSafe> InSourceColorSettings)
 {
 	FreeSample();
 
-	if (!SetProperties(InStride, InWidth, InHeight, InSampleFormat, InTime, InFrameRate, InTimecode, InColorFormatArgs))
+	if (!SetProperties(InStride, InWidth, InHeight, InSampleFormat, InTime, InFrameRate, InTimecode, InSourceColorSettings))
 	{
 		return false;
 	}
@@ -106,7 +107,7 @@ bool FMediaIOCoreTextureSampleBase::SetBuffer(TArray<uint8>&& InVideoBuffer)
 	return true;
 }
 
-bool FMediaIOCoreTextureSampleBase::SetProperties(uint32 InStride, uint32 InWidth, uint32 InHeight, EMediaTextureSampleFormat InSampleFormat, FTimespan InTime, const FFrameRate& InFrameRate, const TOptional<FTimecode>& InTimecode, const UE::MediaIOCore::FColorFormatArgs& InColorFormatArgs)
+bool FMediaIOCoreTextureSampleBase::SetProperties(uint32 InStride, uint32 InWidth, uint32 InHeight, EMediaTextureSampleFormat InSampleFormat, FTimespan InTime, const FFrameRate& InFrameRate, const TOptional<FTimecode>& InTimecode, TSharedPtr<FNativeMediaSourceColorSettings, ESPMode::ThreadSafe> InSourceColorSettings)
 {
 	if (InSampleFormat == EMediaTextureSampleFormat::Undefined)
 	{
@@ -120,18 +121,16 @@ bool FMediaIOCoreTextureSampleBase::SetProperties(uint32 InStride, uint32 InWidt
 	Time = InTime;
 	Duration = FTimespan(ETimespan::TicksPerSecond * InFrameRate.AsInterval());
 	Timecode = InTimecode;
-	Encoding = InColorFormatArgs.Encoding;
-	ColorSpaceType = InColorFormatArgs.ColorSpaceType;
-	ColorSpaceStruct = UE::Color::FColorSpace(ColorSpaceType);
+	SourceColorSettings = InSourceColorSettings;
 
 	return true;
 }
 
-bool FMediaIOCoreTextureSampleBase::InitializeWithEvenOddLine(bool bUseEvenLine, const void* InVideoBuffer, uint32 InBufferSize, uint32 InStride, uint32 InWidth, uint32 InHeight, EMediaTextureSampleFormat InSampleFormat, FTimespan InTime, const FFrameRate& InFrameRate, const TOptional<FTimecode>& InTimecode, const UE::MediaIOCore::FColorFormatArgs& InColorFormatArgs)
+bool FMediaIOCoreTextureSampleBase::InitializeWithEvenOddLine(bool bUseEvenLine, const void* InVideoBuffer, uint32 InBufferSize, uint32 InStride, uint32 InWidth, uint32 InHeight, EMediaTextureSampleFormat InSampleFormat, FTimespan InTime, const FFrameRate& InFrameRate, const TOptional<FTimecode>& InTimecode, TSharedPtr<FNativeMediaSourceColorSettings, ESPMode::ThreadSafe> InSourceColorSettings)
 {
 	FreeSample();
 
-	if (!SetProperties(InStride, InWidth, InHeight/2, InSampleFormat, InTime, InFrameRate, InTimecode, InColorFormatArgs))
+	if (!SetProperties(InStride, InWidth, InHeight/2, InSampleFormat, InTime, InFrameRate, InTimecode, InSourceColorSettings))
 	{
 		return false;
 	}
@@ -231,9 +230,7 @@ void FMediaIOCoreTextureSampleBase::CopyConfiguration(const TSharedPtr<FMediaIOC
 	SampleFormat = SourceSample->SampleFormat;
 	Time = SourceSample->Time;
 	Timecode = SourceSample->Timecode;
-	Encoding = SourceSample->Encoding;
-	ColorSpaceType = SourceSample->ColorSpaceType;
-	ColorSpaceStruct = SourceSample->ColorSpaceStruct;
+	SourceColorSettings = SourceSample->SourceColorSettings;
 	ColorConversionSettings = SourceSample->ColorConversionSettings;
 	CachedOCIOResources = SourceSample->CachedOCIOResources;
 	Player = SourceSample->Player;
@@ -248,6 +245,34 @@ void FMediaIOCoreTextureSampleBase::CopyConfiguration(const TSharedPtr<FMediaIOC
 	OriginalSample = SourceSample;
 }
 
+// Deprecated methods
+PRAGMA_DISABLE_DEPRECATION_WARNINGS
+bool FMediaIOCoreTextureSampleBase::Initialize(const void* InVideoBuffer, uint32 InBufferSize, uint32 InStride, uint32 InWidth, uint32 InHeight, EMediaTextureSampleFormat InSampleFormat, FTimespan InTime, const FFrameRate& InFrameRate, const TOptional<FTimecode>& InTimecode, const UE::MediaIOCore::FColorFormatArgs& InColorFormatArgs)
+{
+	return Initialize(InVideoBuffer, InBufferSize, InStride, InWidth, InHeight, InSampleFormat, InTime, InFrameRate, InTimecode, nullptr);
+}
+
+bool FMediaIOCoreTextureSampleBase::Initialize(const TArray<uint8>& InVideoBuffer, uint32 InStride, uint32 InWidth, uint32 InHeight, EMediaTextureSampleFormat InSampleFormat, FTimespan InTime, const FFrameRate& InFrameRate, const TOptional<FTimecode>& InTimecode, const UE::MediaIOCore::FColorFormatArgs& InColorFormatArgs)
+{
+	return Initialize(InVideoBuffer, InStride, InWidth, InHeight, InSampleFormat, InTime, InFrameRate, InTimecode, nullptr);
+}
+
+bool FMediaIOCoreTextureSampleBase::Initialize(TArray<uint8>&& InVideoBuffer, uint32 InStride, uint32 InWidth, uint32 InHeight, EMediaTextureSampleFormat InSampleFormat, FTimespan InTime, const FFrameRate& InFrameRate, const TOptional<FTimecode>& InTimecode, const UE::MediaIOCore::FColorFormatArgs& InColorFormatArgs)
+{
+	return Initialize(InVideoBuffer, InStride, InWidth, InHeight, InSampleFormat, InTime, InFrameRate, InTimecode, nullptr);
+}
+
+bool FMediaIOCoreTextureSampleBase::SetProperties(uint32 InStride, uint32 InWidth, uint32 InHeight, EMediaTextureSampleFormat InSampleFormat, FTimespan InTime, const FFrameRate& InFrameRate, const TOptional<FTimecode>& InTimecode, const UE::MediaIOCore::FColorFormatArgs& InColorFormatArgs)
+{
+	return SetProperties(InStride, InWidth, InHeight, InSampleFormat, InTime, InFrameRate, InTimecode, nullptr);
+}
+
+bool FMediaIOCoreTextureSampleBase::InitializeWithEvenOddLine(bool bUseEvenLine, const void* InVideoBuffer, uint32 InBufferSize, uint32 InStride, uint32 InWidth, uint32 InHeight, EMediaTextureSampleFormat InSampleFormat, FTimespan InTime, const FFrameRate& InFrameRate, const TOptional<FTimecode>& InTimecode, const UE::MediaIOCore::FColorFormatArgs& InColorFormatArgs)
+{
+	return InitializeWithEvenOddLine(bUseEvenLine, InVideoBuffer, InBufferSize, InStride, InWidth, InHeight, InSampleFormat, InTime, InFrameRate, InTimecode, nullptr);
+}
+PRAGMA_ENABLE_DEPRECATION_WARNINGS
+
 void FMediaIOCoreTextureSampleBase::CacheColorCoversionSettings_GameThread()
 {
 	if (ColorConversionSettings.IsValid() && ColorConversionSettings->IsValid())
@@ -354,15 +379,21 @@ void FMediaIOCoreTextureSampleBase::ShutdownPoolable()
 
 const FMatrix& FMediaIOCoreTextureSampleBase::GetYUVToRGBMatrix() const
 {
-	switch (ColorSpaceType)
+	using namespace UE::Color;
+
+	const FColorSpace& SourceColorSpace = GetSourceColorSpace();
+
+	if (SourceColorSpace.Equals(FColorSpace::GetSRGB()))
 	{
-	case UE::Color::EColorSpace::sRGB:
 		return MediaShaders::YuvToRgbRec709Scaled;
-	case UE::Color::EColorSpace::Rec2020:
+	}
+	
+	if (SourceColorSpace.Equals(UE::Color::FColorSpace::GetRec2020()))
+	{
 		return MediaShaders::YuvToRgbRec2020Scaled;
-	default:
-		return MediaShaders::YuvToRgbRec709Scaled;
 	}
+
+	return MediaShaders::YuvToRgbRec709Scaled;
 }
 
 bool FMediaIOCoreTextureSampleBase::IsOutputSrgb() const
@@ -371,31 +402,68 @@ bool FMediaIOCoreTextureSampleBase::IsOutputSrgb() const
 	{
 		return false; // Don't apply gamma correction as it will be handled by the OCIO conversion.
 	}
-	return Encoding == UE::Color::EEncoding::sRGB;
+
+	if (SourceColorSettings.IsValid())
+	{
+		const UE::Color::EEncoding EncodingOverride = SourceColorSettings->GetEncodingOverride();
+
+		if (EncodingOverride != UE::Color::EEncoding::None)
+		{
+			return EncodingOverride == UE::Color::EEncoding::sRGB;
+		}
+	}
+
+	return false;
 }
 
 const UE::Color::FColorSpace& FMediaIOCoreTextureSampleBase::GetSourceColorSpace() const
 {
-	return ColorSpaceStruct;
+	if (SourceColorSettings.IsValid())
+	{
+		return SourceColorSettings->GetColorSpaceOverride(UE::Color::FColorSpace::GetSRGB());
+	}
+
+	return IMediaTextureSample::GetSourceColorSpace();
+}
+
+bool FMediaIOCoreTextureSampleBase::HasEncodingOverride() const
+{
+	return SourceColorSettings.IsValid() ? SourceColorSettings->HasEncodingOverride() : false;
+}
+
+bool FMediaIOCoreTextureSampleBase::HasColorSpaceOverride() const
+{
+	return SourceColorSettings.IsValid() ? SourceColorSettings->HasColorSpaceOverride() : false;
 }
 
 UE::Color::EEncoding FMediaIOCoreTextureSampleBase::GetEncodingType() const
 {
 	if (ColorConversionSettings && ColorConversionSettings->IsValid())
-    {
-    	return UE::Color::EEncoding::Linear;
-    }
-    
-	return Encoding;
+	{
+		return UE::Color::EEncoding::Linear;
+	}
+
+	if (SourceColorSettings.IsValid())
+	{
+		const UE::Color::EEncoding EncodingOverride = SourceColorSettings->GetEncodingOverride();
+
+		if (EncodingOverride != UE::Color::EEncoding::None)
+		{
+			return EncodingOverride;
+		}
+	}
+
+	return IMediaTextureSample::GetEncodingType();
 }
 
-UE::Color::EColorSpace FMediaIOCoreTextureSampleBase::GetColorSpaceType() const
+UE::Color::EChromaticAdaptationMethod FMediaIOCoreTextureSampleBase::GetChromaticAdapationMethod() const
 {
-	if (ColorConversionSettings && ColorConversionSettings->IsValid())
+	if (SourceColorSettings.IsValid())
 	{
-		return UE::Color::EColorSpace::None;
+		return SourceColorSettings->GetChromaticAdaptationMethod();
 	}
-	return ColorSpaceType;
+
+	return IMediaTextureSample::GetChromaticAdapationMethod();
 }
 
 float FMediaIOCoreTextureSampleBase::GetHDRNitsNormalizationFactor() const
@@ -406,4 +474,10 @@ float FMediaIOCoreTextureSampleBase::GetHDRNitsNormalizationFactor() const
     }
 	
 	return IMediaTextureSample::GetHDRNitsNormalizationFactor();
-}
\ No newline at end of file
+}
+
+// Deprecated
+UE::Color::EColorSpace FMediaIOCoreTextureSampleBase::GetColorSpaceType() const
+{
+	return UE::Color::EColorSpace::None;
+}
diff --git a/Engine/Plugins/Media/MediaIOFramework/Source/MediaIOCore/Private/Tests/MediaIODeinterlacerTests.cpp b/Engine/Plugins/Media/MediaIOFramework/Source/MediaIOCore/Private/Tests/MediaIODeinterlacerTests.cpp
--- a/Engine/Plugins/Media/MediaIOFramework/Source/MediaIOCore/Private/Tests/MediaIODeinterlacerTests.cpp
+++ b/Engine/Plugins/Media/MediaIOFramework/Source/MediaIOCore/Private/Tests/MediaIODeinterlacerTests.cpp
@@ -57,7 +57,7 @@ namespace UE::MediaIOTests
 			FTimespan::FromSeconds(0),
 			FFrameRate(),
 			FTimecode(),
-			UE::MediaIOCore::FColorFormatArgs{UE::Color::EEncoding::Linear, UE::Color::EColorSpace::sRGB}
+			nullptr
 		};
 
 		// Execute test
@@ -98,7 +98,7 @@ namespace UE::MediaIOTests
 			FTimespan::FromSeconds(0),
 			FFrameRate(),
 			FTimecode(),
-			UE::MediaIOCore::FColorFormatArgs{UE::Color::EEncoding::Linear, UE::Color::EColorSpace::sRGB}
+			nullptr
 		};
 
 		// Execute test
@@ -134,7 +134,7 @@ namespace UE::MediaIOTests
 			FTimespan::FromSeconds(0),
 			FFrameRate(),
 			FTimecode(),
-			UE::MediaIOCore::FColorFormatArgs{UE::Color::EEncoding::Linear, UE::Color::EColorSpace::sRGB}
+			nullptr
 		};
 
 		// Execute test
@@ -171,7 +171,7 @@ namespace UE::MediaIOTests
 			FTimespan::FromSeconds(0),
 			FFrameRate(),
 			FTimecode(),
-			UE::MediaIOCore::FColorFormatArgs{UE::Color::EEncoding::Linear, UE::Color::EColorSpace::sRGB}
+			nullptr
 		};
 
 		// Execute test
diff --git a/Engine/Plugins/Media/MediaIOFramework/Source/MediaIOCore/Public/CaptureCardMediaSource.h b/Engine/Plugins/Media/MediaIOFramework/Source/MediaIOCore/Public/CaptureCardMediaSource.h
--- a/Engine/Plugins/Media/MediaIOFramework/Source/MediaIOCore/Public/CaptureCardMediaSource.h
+++ b/Engine/Plugins/Media/MediaIOFramework/Source/MediaIOCore/Public/CaptureCardMediaSource.h
@@ -3,6 +3,7 @@
 #pragma once
 
 #include "TimeSynchronizableMediaSource.h"
+#include "BaseMediaSourceColorSettings.h"
 #include "Containers/UnrealString.h"
 #include "Engine/TextureDefines.h"
 #include "HAL/Platform.h"
@@ -27,22 +28,30 @@ namespace UE::CaptureCardMediaSource
 	static FName EvaluationType("EvaluationType");
 	static FName Deinterlacer("Deinterlacer");
 	static FName InterlaceFieldOrder("InterlaceFieldOrder");
-	static FName OverrideSourceEncoding("OverrideSourceEncoding");
-	static FName SourceEncoding("SourceEncoding");
-	static FName OverrideSourceColorSpace("OverrideSourceColorSpace");
-	static FName SourceColorSpace("SourceColorcSpace");
+	static FName SourceColorSettingsOption("SourceColorSettings");
 	static FName OpenColorIOSettings("OpenColorIOSettings");
+
+	// 5.8 deprecated options
+	UE_DEPRECATED(5.8, "OverrideSourceEncoding is deprecated.")
+	static FName OverrideSourceEncoding("OverrideSourceEncoding_DEPRECATED");
+
+	UE_DEPRECATED(5.8, "SourceEncoding is deprecated.")
+	static FName SourceEncoding("SourceEncoding_DEPRECATED");
+
+	UE_DEPRECATED(5.8, "OverrideSourceColorSpace is deprecated.")
+	static FName OverrideSourceColorSpace("OverrideSourceColorSpace_DEPRECATED");
+
+	UE_DEPRECATED(5.8, "SourceColorSpace is deprecated.")
+	static FName SourceColorSpace("SourceColorcSpace_DEPRECATED");
 }
 
 
-/** List of texture source encodings that can be converted to linear. (Integer values match the ETextureSourceEncoding values in TextureDefines.h */
-UENUM()
-enum class EMediaIOCoreSourceEncoding : uint8
+UENUM(meta = (Deprecated = "5.8", DeprecationMessage = "EMediaIOCoreSourceEncoding is deprecated. This functionality is now available with EMediaSourceEncoding in BaseMediaSourceColorSettings.h"))
+enum class UE_DEPRECATED(5.8, "EMediaIOCoreSourceEncoding is deprecated. This functionality is now available with EMediaSourceEncoding in BaseMediaSourceColorSettings.") EMediaIOCoreSourceEncoding : uint8
 {
 	Linear		= 1 UMETA(DisplayName = "Linear"),
 	sRGB		= 2 UMETA(DisplayName = "sRGB"),
 	ST2084		= 3 UMETA(DisplayName = "ST 2084/PQ"),
-	//BT1886		= 5 UMETA(DisplayName = "BT1886/Gamma 2.4"),
 	SLog3		= 12 UMETA(DisplayName = "SLog3"),
 	MAX,
 };
@@ -84,28 +93,20 @@ public:
 	UPROPERTY(BlueprintReadOnly, EditAnywhere, Category = "Video")
 	EMediaIOInterlaceFieldOrder InterlaceFieldOrder = EMediaIOInterlaceFieldOrder::TopFieldFirst;
 
-	/** Whether to override the source encoding or to use the metadata embedded in the ancillary data of the signal. */
-	UPROPERTY(EditAnywhere, Category = "Video", meta = (InlineEditConditionToggle))
-	bool bOverrideSourceEncoding = true;
-
-	/** Encoding of the source texture. */
-	UPROPERTY(BlueprintReadOnly, EditAnywhere, Category = "Video", meta = (EditCondition = "bOverrideSourceEncoding"))
-	EMediaIOCoreSourceEncoding OverrideSourceEncoding = EMediaIOCoreSourceEncoding::Linear;
-
-	/** Whether to override the source color space or to use the metadata embedded in the ancillary data of the signal. */
-	UPROPERTY(EditAnywhere, Category = "Video", meta = (InlineEditConditionToggle))
-	bool bOverrideSourceColorSpace = true;
-
-	/** Color space of the source texture. */
-	UPROPERTY(BlueprintReadOnly, EditAnywhere, Category = "Video", meta = (EditCondition = "bOverrideSourceColorSpace", InvalidEnumValues = "ETextureColorSpace::TCS_Custom"))
-	ETextureColorSpace OverrideSourceColorSpace = ETextureColorSpace::TCS_sRGB;
-
 	/**
-	 * OCIO Settings used for applying a color conversion to the incoming source.
+	 * OpenColorIO transform settings used for applying a color conversion override to the incoming source.
 	 */
-	UPROPERTY(BlueprintReadWrite ,EditAnywhere, Category="Video")
+	UPROPERTY(BlueprintReadWrite, EditAnywhere, Category = "ColorManagement", meta = (DisplayName = "OpenColorIO Override"))
 	FOpenColorIOColorConversionSettings ColorConversionSettings;
 
+	/** Manual definition of media source color space & encoding. */
+	UPROPERTY(BlueprintReadWrite, EditAnywhere, Category = "ColorManagement", meta = (ShowOnlyInnerProperties))
+	FMediaSourceColorSettings SourceColorSettings;
+
+private:
+	/** Native source color settings. */
+	TSharedPtr<FNativeMediaSourceColorSettings, ESPMode::ThreadSafe> NativeSourceColorSettings;
+
 public:
 	struct FOpenColorIODataContainer : public IMediaOptions::FDataContainer
 	{
@@ -122,10 +123,33 @@ public:
 	UE_API virtual bool HasMediaOption(const FName& Key) const override;
 
 public:
+	UE_API virtual void PostLoad() override;
 #if WITH_EDITOR
 	UE_API virtual bool CanEditChange(const FProperty* InProperty) const override;
 	UE_API virtual void PostEditChangeChainProperty(struct FPropertyChangedChainEvent& PropertyChangedEvent) override;
 #endif //WITH_EDITOR
+
+public:
+	//~ Deprecated properties
+#if WITH_EDITORONLY_DATA
+PRAGMA_DISABLE_DEPRECATION_WARNINGS
+	UE_DEPRECATED(5.8, "This property is deprecated. Please use SourceColorSettings instead.")
+	UPROPERTY(meta = (DeprecatedProperty, DeprecationMessage = "This property is deprecated. Please use SourceColorSettings instead."))
+	bool bOverrideSourceEncoding_DEPRECATED = true;
+
+	UE_DEPRECATED(5.8, "This property is deprecated. Please use SourceColorSettings instead.")
+	UPROPERTY(meta = (DeprecatedProperty, DeprecationMessage = "This property is deprecated. Please use SourceColorSettings instead."))
+	EMediaIOCoreSourceEncoding OverrideSourceEncoding_DEPRECATED = EMediaIOCoreSourceEncoding::Linear;
+
+	UE_DEPRECATED(5.8, "This property is deprecated. Please use SourceColorSettings instead.")
+	UPROPERTY(meta = (DeprecatedProperty, DeprecationMessage = "This property is deprecated. Please use SourceColorSettings instead."))
+	bool bOverrideSourceColorSpace_DEPRECATED = true;
+
+	UE_DEPRECATED(5.8, "This property is deprecated. Please use SourceColorSettings instead.")
+	UPROPERTY(meta = (DeprecatedProperty, DeprecationMessage = "This property is deprecated. Please use SourceColorSettings instead."))
+	ETextureColorSpace OverrideSourceColorSpace_DEPRECATED = ETextureColorSpace::TCS_sRGB;
+PRAGMA_ENABLE_DEPRECATION_WARNINGS
+#endif
 };
 
 #undef UE_API
diff --git a/Engine/Plugins/Media/MediaIOFramework/Source/MediaIOCore/Public/MediaIOCoreDeinterlacer.h b/Engine/Plugins/Media/MediaIOFramework/Source/MediaIOCore/Public/MediaIOCoreDeinterlacer.h
--- a/Engine/Plugins/Media/MediaIOFramework/Source/MediaIOCore/Public/MediaIOCoreDeinterlacer.h
+++ b/Engine/Plugins/Media/MediaIOFramework/Source/MediaIOCore/Public/MediaIOCoreDeinterlacer.h
@@ -25,7 +25,7 @@ namespace UE::MediaIOCore
 		, FTimespan InTime
 		, const FFrameRate& InFrameRate
 		, const TOptional<FTimecode>& InTimecode
-		, const UE::MediaIOCore::FColorFormatArgs& InColorFormatArgs)
+		, TSharedPtr<FNativeMediaSourceColorSettings, ESPMode::ThreadSafe> InSourceColorSettings)
 			: VideoBuffer(InVideoBuffer)
 			, BufferSize(InBufferSize)
 			, Stride(InStride)
@@ -35,7 +35,7 @@ namespace UE::MediaIOCore
 			, Time(InTime)
 			, FrameRate(InFrameRate)
 			, Timecode(InTimecode)
-			, ColorFormatArgs(InColorFormatArgs)
+			, SourceColorSettings(InSourceColorSettings)
 		{
 		}
 
@@ -48,7 +48,34 @@ namespace UE::MediaIOCore
 		FTimespan Time;
 		const FFrameRate& FrameRate;
 		const TOptional<FTimecode>& Timecode;
-		const UE::MediaIOCore::FColorFormatArgs& ColorFormatArgs;
+		TSharedPtr<FNativeMediaSourceColorSettings, ESPMode::ThreadSafe> SourceColorSettings;
+
+		UE_DEPRECATED(5.8, "Deprecated constructor. Please use FNativeMediaSourceColorSettings variant.")
+		FVideoFrame(const void* InVideoBuffer
+			, uint32 InBufferSize
+			, uint32 InStride
+			, uint32 InWidth
+			, uint32 InHeight
+			, EMediaTextureSampleFormat InSampleFormat
+			, FTimespan InTime
+			, const FFrameRate& InFrameRate
+			, const TOptional<FTimecode>& InTimecode
+PRAGMA_DISABLE_DEPRECATION_WARNINGS
+			, const UE::MediaIOCore::FColorFormatArgs& InColorFormatArgs
+PRAGMA_ENABLE_DEPRECATION_WARNINGS
+		)
+			: VideoBuffer(InVideoBuffer)
+			, BufferSize(InBufferSize)
+			, Stride(InStride)
+			, Width(InWidth)
+			, Height(InHeight)
+			, SampleFormat(InSampleFormat)
+			, Time(InTime)
+			, FrameRate(InFrameRate)
+			, Timecode(InTimecode)
+			, SourceColorSettings(nullptr)
+		{
+		}
 	};
 
 	/**
diff --git a/Engine/Plugins/Media/MediaIOFramework/Source/MediaIOCore/Public/MediaIOCorePlayerBase.h b/Engine/Plugins/Media/MediaIOFramework/Source/MediaIOCore/Public/MediaIOCorePlayerBase.h
--- a/Engine/Plugins/Media/MediaIOFramework/Source/MediaIOCore/Public/MediaIOCorePlayerBase.h
+++ b/Engine/Plugins/Media/MediaIOFramework/Source/MediaIOCore/Public/MediaIOCorePlayerBase.h
@@ -4,6 +4,7 @@
 
 #include "CoreMinimal.h"
 
+#include "BaseMediaSourceColorSettings.h"
 #include "Engine/TextureDefines.h"
 #include "IMediaCache.h"
 #include "IMediaControls.h"
@@ -174,10 +175,6 @@ public:
 		double EvaluationOffset = 0.;
 		uint32 FrameNumber = 0.;
 	};
-
-	/** Deprecated in UE5.5. */
-	UE_DEPRECATED("5.5", "Use the overloaded method that takes FFrameInfo instead of proxy sample.")
-	bool JustInTimeSampleRender_RenderThread(FRHICommandListImmediate& RHICmdList, TSharedPtr<FMediaIOCoreTextureSampleBase>& JITRProxySample) { return false; };
 	
 	/**
 	 * Just in time sample rendering. This method is responsible for late sample picking,
@@ -338,18 +335,23 @@ protected:
 
 	/** Used to ensure that JIT rendering is executed only once per frame */
 	uint64 LastEngineRTFrameThatUpdatedJustInTime = TNumericLimits<uint64>::Max();
+	
+	/** Container for media source color settings. */
+	TSharedPtr<FNativeMediaSourceColorSettings, ESPMode::ThreadSafe> SourceColorSettings;
 
-	/** Whether to override the source encoding or to use the metadata embedded in the ancillary data of the signal. */
+#if WITH_EDITORONLY_DATA
+	UE_DEPRECATED(5.8, "This property is deprecated. Please use SourceColorSettings instead.")
 	bool bOverrideSourceEncoding = true;
 
-	/** Encoding of the source texture. */
+	UE_DEPRECATED(5.8, "This property is deprecated. Please use SourceColorSettings instead.")
 	ETextureSourceEncoding OverrideSourceEncoding = ETextureSourceEncoding::TSE_Linear;
 
-	/** Whether to override the source color space or to use the metadata embedded in the ancillary data of the signal. */
+	UE_DEPRECATED(5.8, "This property is deprecated. Please use SourceColorSettings instead.")
 	bool bOverrideSourceColorSpace = true;
 
-	/** Color space of the source texture. */
+	UE_DEPRECATED(5.8, "This property is deprecated. Please use SourceColorSettings instead.")
 	ETextureColorSpace OverrideSourceColorSpace = ETextureColorSpace::TCS_None;
+#endif
 
 private:
 	UE_API void OnSampleDestroyed(TRefCountPtr<FRHITexture> InTexture);
diff --git a/Engine/Plugins/Media/MediaIOFramework/Source/MediaIOCore/Public/MediaIOCoreTextureSampleBase.h b/Engine/Plugins/Media/MediaIOFramework/Source/MediaIOCore/Public/MediaIOCoreTextureSampleBase.h
--- a/Engine/Plugins/Media/MediaIOFramework/Source/MediaIOCore/Public/MediaIOCoreTextureSampleBase.h
+++ b/Engine/Plugins/Media/MediaIOFramework/Source/MediaIOCore/Public/MediaIOCoreTextureSampleBase.h
@@ -3,12 +3,10 @@
 #pragma once
 
 #include "IMediaTextureSample.h"
+#include "BaseMediaSourceColorSettings.h"
 #include "MediaIOCoreDefinitions.h"
 #include "IMediaTextureSampleConverter.h"
 #include "MediaObjectPool.h"
-
-
-#include "ColorManagement/ColorSpace.h"
 #include "Misc/FrameRate.h"
 #include "Templates/RefCounting.h"
 
@@ -20,7 +18,7 @@ class FRHITexture;
 
 namespace UE::MediaIOCore
 {
-	struct FColorFormatArgs
+	struct UE_DEPRECATED(5.8, "FColorFormatArgs is deprecated. Please use FNativeMediaSourceColorSettings instead.") FColorFormatArgs
 	{
 		FColorFormatArgs() = default;
 		
@@ -110,9 +108,9 @@ public:
 	 * @param InTime The sample time (in the player's own clock).
 	 * @param InFrameRate The framerate of the media that produce the sample.
 	 * @param InTimecode The sample timecode if available.
-	 * @param InColorFormatArgs Information about the texture color encoding and color space.
+	 * @param InSourceColorSettings Information about the texture color encoding and color space.
 	 */
-	UE_API bool Initialize(const void* InVideoBuffer, uint32 InBufferSize, uint32 InStride, uint32 InWidth, uint32 InHeight, EMediaTextureSampleFormat InSampleFormat, FTimespan InTime, const FFrameRate& InFrameRate, const TOptional<FTimecode>& InTimecode, const UE::MediaIOCore::FColorFormatArgs& InColorFormatArgs);
+	UE_API bool Initialize(const void* InVideoBuffer, uint32 InBufferSize, uint32 InStride, uint32 InWidth, uint32 InHeight, EMediaTextureSampleFormat InSampleFormat, FTimespan InTime, const FFrameRate& InFrameRate, const TOptional<FTimecode>& InTimecode, TSharedPtr<FNativeMediaSourceColorSettings, ESPMode::ThreadSafe> InSourceColorSettings);
 
 	/**
 	 * Initialize the sample.
@@ -125,9 +123,9 @@ public:
 	 * @param InTime The sample time (in the player's own clock).
 	 * @param InFrameRate The framerate of the media that produce the sample.
 	 * @param InTimecode The sample timecode if available.
-	 * @param InColorFormatArgs Information about the texture color encoding and color space.
+	 * @param InSourceColorSettings Information about the texture color encoding and color space.
 	 */
-	UE_API bool Initialize(const TArray<uint8>& InVideoBuffer, uint32 InStride, uint32 InWidth, uint32 InHeight, EMediaTextureSampleFormat InSampleFormat, FTimespan InTime, const FFrameRate& InFrameRate, const TOptional<FTimecode>& InTimecode, const UE::MediaIOCore::FColorFormatArgs& InColorFormatArgs);
+	UE_API bool Initialize(const TArray<uint8>& InVideoBuffer, uint32 InStride, uint32 InWidth, uint32 InHeight, EMediaTextureSampleFormat InSampleFormat, FTimespan InTime, const FFrameRate& InFrameRate, const TOptional<FTimecode>& InTimecode, TSharedPtr<FNativeMediaSourceColorSettings, ESPMode::ThreadSafe> InSourceColorSettings);
 
 	/**
 	 * Initialize the sample.
@@ -140,9 +138,9 @@ public:
 	 * @param InTime The sample time (in the player's own clock).
 	 * @param InFrameRate The framerate of the media that produce the sample.
 	 * @param InTimecode The sample timecode if available.
-	 * @param InColorFormatArgs Information about the texture color encoding and color space.
+	 * @param InSourceColorSettings Information about the texture color encoding and color space.
 	 */
-	UE_API bool Initialize(TArray<uint8>&& InVideoBuffer, uint32 InStride, uint32 InWidth, uint32 InHeight, EMediaTextureSampleFormat InSampleFormat, FTimespan InTime, const FFrameRate& InFrameRate, const TOptional<FTimecode>& InTimecode, const UE::MediaIOCore::FColorFormatArgs& InColorFormatArgs);
+	UE_API bool Initialize(TArray<uint8>&& InVideoBuffer, uint32 InStride, uint32 InWidth, uint32 InHeight, EMediaTextureSampleFormat InSampleFormat, FTimespan InTime, const FFrameRate& InFrameRate, const TOptional<FTimecode>& InTimecode, TSharedPtr<FNativeMediaSourceColorSettings, ESPMode::ThreadSafe> InSourceColorSettings);
 
 	/**
 	 * Initialize the sample.
@@ -170,15 +168,15 @@ public:
 	 * Set the sample properties.
 	 *
 	 * @param InStride The number of channel of the video buffer.
-	 * @param InWidth The sample rate of the video buffer.
-	 * @param InHeight The sample rate of the video buffer.
+	 * @param InWidth The width the video buffer.
+	 * @param InHeight The height of the video buffer.
 	 * @param InSampleFormat The sample format of the video buffer.
 	 * @param InTime The sample time (in the player's own clock).
 	 * @param InFrameRate The framerate of the media that produce the sample.
 	 * @param InTimecode The sample timecode if available.
-	 * @param InColorFormatArgs Information about the texture color encoding and color space.
+	 * @param InSourceColorSettings Information about the texture color encoding and color space.
 	 */
-	UE_API bool SetProperties(uint32 InStride, uint32 InWidth, uint32 InHeight, EMediaTextureSampleFormat InSampleFormat, FTimespan InTime, const FFrameRate& InFrameRate, const TOptional<FTimecode>& InTimecode, const UE::MediaIOCore::FColorFormatArgs& InColorFormatArgs);
+	UE_API bool SetProperties(uint32 InStride, uint32 InWidth, uint32 InHeight, EMediaTextureSampleFormat InSampleFormat, FTimespan InTime, const FFrameRate& InFrameRate, const TOptional<FTimecode>& InTimecode, TSharedPtr<FNativeMediaSourceColorSettings, ESPMode::ThreadSafe> InSourceColorSettings);
 
 	/**
 	 * Initialize the sample with half it's original height and take only the odd or even line.
@@ -186,15 +184,15 @@ public:
 	 * @param bUseEvenLine Should use the Even or the Odd line from the video buffer.
 	 * @param InVideoBuffer The video frame data.
 	 * @param InStride The number of channel of the video buffer.
-	 * @param InWidth The sample rate of the video buffer.
-	 * @param InHeight The sample rate of the video buffer.
+	 * @param InWidth The width of the video buffer.
+	 * @param InHeight The height of the video buffer.
 	 * @param InSampleFormat The sample format of the video buffer.
 	 * @param InTime The sample time (in the player's own clock).
 	 * @param InTimecode The sample timecode if available.
-	 * @param InColorFormatArgs Information about the texture color encoding and color space.
+	 * @param InSourceColorSettings Information about the texture color encoding and color space.
 	 */
-	UE_API bool InitializeWithEvenOddLine(bool bUseEvenLine, const void* InVideoBuffer, uint32 InBufferSize, uint32 InStride, uint32 InWidth, uint32 InHeight, EMediaTextureSampleFormat InSampleFormat, FTimespan InTime, const FFrameRate& InFrameRate, const TOptional<FTimecode>& InTimecode, const UE::MediaIOCore::FColorFormatArgs& InColorFormatArgs);
-
+	UE_API bool InitializeWithEvenOddLine(bool bUseEvenLine, const void* InVideoBuffer, uint32 InBufferSize, uint32 InStride, uint32 InWidth, uint32 InHeight, EMediaTextureSampleFormat InSampleFormat, FTimespan InTime, const FFrameRate& InFrameRate, const TOptional<FTimecode>& InTimecode, TSharedPtr<FNativeMediaSourceColorSettings, ESPMode::ThreadSafe> InSourceColorSettings);
+	
 	/**
 	 * Set the sample buffer with half it's original height and take only the odd or even line.
 	 *
@@ -202,7 +200,7 @@ public:
 	 * @param InVideoBuffer The video frame data.
 	 * @param InBufferSize The size of the video buffer.
 	 * @param InStride The number of channel of the video buffer.
-	 * @param InHeight The sample rate of the video buffer.
+	 * @param InHeight The height of the video buffer.
 	 */
 	UE_API bool SetBufferWithEvenOddLine(bool bUseEvenLine, const void* InVideoBuffer, uint32 InBufferSize, uint32 InStride, uint32 InHeight);
 
@@ -251,6 +249,21 @@ public:
 
 	/** Copies all neccessary data from a source sample to render JIT */
 	UE_API virtual void CopyConfiguration(const TSharedPtr<FMediaIOCoreTextureSampleBase>& SourceSample);
+
+	//~ Deprecated methods
+PRAGMA_DISABLE_DEPRECATION_WARNINGS
+	UE_DEPRECATED(5.8, "Deprecated method. Please use method with FNativeMediaSourceColorSettings argument instead.")
+	UE_API bool Initialize(const void* InVideoBuffer, uint32 InBufferSize, uint32 InStride, uint32 InWidth, uint32 InHeight, EMediaTextureSampleFormat InSampleFormat, FTimespan InTime, const FFrameRate& InFrameRate, const TOptional<FTimecode>& InTimecode, const UE::MediaIOCore::FColorFormatArgs& InColorFormatArgs);
+	UE_DEPRECATED(5.8, "Deprecated method. Please use method with FNativeMediaSourceColorSettings argument instead.")
+	UE_API bool Initialize(const TArray<uint8>& InVideoBuffer, uint32 InStride, uint32 InWidth, uint32 InHeight, EMediaTextureSampleFormat InSampleFormat, FTimespan InTime, const FFrameRate& InFrameRate, const TOptional<FTimecode>& InTimecode, const UE::MediaIOCore::FColorFormatArgs& InColorFormatArgs);
+	UE_DEPRECATED(5.8, "Deprecated method. Please use method with FNativeMediaSourceColorSettings argument instead.")
+	UE_API bool Initialize(TArray<uint8>&& InVideoBuffer, uint32 InStride, uint32 InWidth, uint32 InHeight, EMediaTextureSampleFormat InSampleFormat, FTimespan InTime, const FFrameRate& InFrameRate, const TOptional<FTimecode>& InTimecode, const UE::MediaIOCore::FColorFormatArgs& InColorFormatArgs);
+	UE_DEPRECATED(5.8, "Deprecated method. Please use method with FNativeMediaSourceColorSettings argument instead.")
+	UE_API bool SetProperties(uint32 InStride, uint32 InWidth, uint32 InHeight, EMediaTextureSampleFormat InSampleFormat, FTimespan InTime, const FFrameRate& InFrameRate, const TOptional<FTimecode>& InTimecode, const UE::MediaIOCore::FColorFormatArgs& InColorFormatArgs);
+	UE_DEPRECATED(5.8, "Deprecated method. Please use method with FNativeMediaSourceColorSettings argument instead.")
+	UE_API bool InitializeWithEvenOddLine(bool bUseEvenLine, const void* InVideoBuffer, uint32 InBufferSize, uint32 InStride, uint32 InWidth, uint32 InHeight, EMediaTextureSampleFormat InSampleFormat, FTimespan InTime, const FFrameRate& InFrameRate, const TOptional<FTimecode>& InTimecode, const UE::MediaIOCore::FColorFormatArgs& InColorFormatArgs);
+PRAGMA_ENABLE_DEPRECATION_WARNINGS
+
 protected:
 	/**
 	* Method that caches color conversion settings on Game thread.
@@ -316,12 +329,15 @@ public:
 	}
 
 	UE_API virtual bool IsOutputSrgb() const override;
-	
+
+	UE_API virtual bool HasEncodingOverride() const;
+	UE_API virtual bool HasColorSpaceOverride() const;
 	UE_API virtual const UE::Color::FColorSpace& GetSourceColorSpace() const override;
 	UE_API virtual UE::Color::EEncoding GetEncodingType() const override;
+	UE_API virtual UE::Color::EChromaticAdaptationMethod GetChromaticAdapationMethod() const override;
 	UE_API virtual float GetHDRNitsNormalizationFactor() const override;
 	
-	/**Method that returns this sample's colorspace type.*/
+	UE_DEPRECATED(5.8, "GetColorSpaceType is deprecated, please use GetSourceColorSpace instead.")
 	UE_API virtual UE::Color::EColorSpace GetColorSpaceType() const;
 
 	virtual const void* GetBuffer() override
@@ -447,19 +463,15 @@ protected:
 	void* ExternalBuffer = nullptr;
 
 #if WITH_EDITORONLY_DATA
-	UE_DEPRECATED(5.3, "Please use Encoding instead.")
-	/** Whether the sample is in sRGB space and requires an explicit conversion to linear */
-	bool bIsSRGBInput = false;
-#endif
-
-	/** Color encoding of the incoming texture. */
+	UE_DEPRECATED(5.8, "Please use SourceColorSettings instead.")
 	UE::Color::EEncoding Encoding = UE::Color::EEncoding::Linear;
 
-	/** Color space enum of the incoming texture. */
+	UE_DEPRECATED(5.8, "Please use SourceColorSettings instead.")
 	UE::Color::EColorSpace ColorSpaceType = UE::Color::EColorSpace::sRGB;
 
-	/** Color space structure of the incoming texture. Used for retrieving chromaticities. */
+	UE_DEPRECATED(5.8, "Please use SourceColorSettings instead.")
 	UE::Color::FColorSpace ColorSpaceStruct = UE::Color::FColorSpace(UE::Color::EColorSpace::sRGB);
+#endif
 
 	/** The player that created this sample */
 	TWeakPtr<FMediaIOCorePlayerBase> Player;
@@ -481,6 +493,8 @@ private:
 	double EvaluationOffsetInSeconds = 0;
 
 private:
+	/** Manual source color space & encoding overrides. */
+	TSharedPtr<FNativeMediaSourceColorSettings, ESPMode::ThreadSafe> SourceColorSettings;
 	/** Settings used to apply an OCIO conversion to the sample. */
 	TSharedPtr<struct FOpenColorIOColorConversionSettings> ColorConversionSettings;
 	/** Cached rendering resources used for the color conversion pass when using OCIO. */
diff --git a/Engine/Plugins/Media/NDIMedia/Source/NDIMedia/Private/NDIMediaSource.cpp b/Engine/Plugins/Media/NDIMedia/Source/NDIMedia/Private/NDIMediaSource.cpp
--- a/Engine/Plugins/Media/NDIMedia/Source/NDIMedia/Private/NDIMediaSource.cpp
+++ b/Engine/Plugins/Media/NDIMedia/Source/NDIMedia/Private/NDIMediaSource.cpp
@@ -22,12 +22,7 @@ UNDIMediaSource::UNDIMediaSource()
 	, MaxNumVideoFrameBuffer(8)
 	, bLogDropFrame(true)
 	, bEncodeTimecodeInTexel(false)
-{
-	bOverrideSourceEncoding = false;
-	OverrideSourceEncoding = EMediaIOCoreSourceEncoding::sRGB;
-	bOverrideSourceColorSpace = false;
-	OverrideSourceColorSpace = ETextureColorSpace::TCS_sRGB;
-	
+{	
 	AssignDefaultConfiguration();
 }
 
diff --git a/Engine/Plugins/Media/NDIMedia/Source/NDIMedia/Private/Player/NDIMediaStreamPlayer.cpp b/Engine/Plugins/Media/NDIMedia/Source/NDIMedia/Private/Player/NDIMediaStreamPlayer.cpp
--- a/Engine/Plugins/Media/NDIMedia/Source/NDIMedia/Private/Player/NDIMediaStreamPlayer.cpp
+++ b/Engine/Plugins/Media/NDIMedia/Source/NDIMedia/Private/Player/NDIMediaStreamPlayer.cpp
@@ -307,14 +307,10 @@ void FNDIMediaStreamPlayer::TickFetch(FTimespan InDeltaTime, FTimespan InTime)
 void FNDIMediaStreamPlayer::HandleVideoFrameReceived(FNDIStreamReceiver* InReceiver, const NDIlib_video_frame_v2_t& InVideoFrame, const FTimespan& InTime)
 {
 	TSharedRef<FNDIMediaTextureSample> TextureSample = TextureSamplePool->AcquireShared();
-	
-	const UE::MediaIOCore::FColorFormatArgs ColorFormatArgs(
-		bOverrideSourceEncoding ?  static_cast<UE::Color::EEncoding>(OverrideSourceEncoding) : UE::Color::EEncoding::sRGB,
-		bOverrideSourceColorSpace ? static_cast<UE::Color::EColorSpace>(OverrideSourceColorSpace) : UE::Color::EColorSpace::sRGB);
 
 	FTimecode SourceTimecode = InReceiver->GetCurrentTimecode();
 
-	if (TextureSample->Initialize(InVideoFrame, ColorFormatArgs, InTime, SourceTimecode))
+	if (TextureSample->Initialize(InVideoFrame, SourceColorSettings, InTime, SourceTimecode))
 	{
 		if (TextureSample->CustomConverter)
 		{
diff --git a/Engine/Plugins/Media/NDIMedia/Source/NDIMedia/Private/Player/NDIMediaTextureSample.cpp b/Engine/Plugins/Media/NDIMedia/Source/NDIMedia/Private/Player/NDIMediaTextureSample.cpp
--- a/Engine/Plugins/Media/NDIMedia/Source/NDIMedia/Private/Player/NDIMediaTextureSample.cpp
+++ b/Engine/Plugins/Media/NDIMedia/Source/NDIMedia/Private/Player/NDIMediaTextureSample.cpp
@@ -5,7 +5,7 @@
 #include "NDIMediaAPI.h"
 #include "NDIMediaTextureSampleConverter.h"
 
-bool FNDIMediaTextureSample::Initialize(const NDIlib_video_frame_v2_t& InVideoFrame, const UE::MediaIOCore::FColorFormatArgs& InColorFormatArgs, const FTimespan& InTime, const TOptional<FTimecode>& InTimecode)
+bool FNDIMediaTextureSample::Initialize(const NDIlib_video_frame_v2_t& InVideoFrame, TSharedPtr<FNativeMediaSourceColorSettings, ESPMode::ThreadSafe> InSourceColorSettings, const FTimespan& InTime, const TOptional<FTimecode>& InTimecode)
 {
 	bIsCustomFormat = false;
 	bIsProgressive = true;
@@ -61,7 +61,7 @@ bool FNDIMediaTextureSample::Initialize(const NDIlib_video_frame_v2_t& InVideoFr
 				, InTime
 				, FFrameRate(InVideoFrame.frame_rate_N, InVideoFrame.frame_rate_D)
 				, InTimecode
-				, InColorFormatArgs);
+				, InSourceColorSettings);
 
 		case NDIlib_frame_format_type_field_0:
 		case NDIlib_frame_format_type_field_1:
@@ -78,7 +78,7 @@ bool FNDIMediaTextureSample::Initialize(const NDIlib_video_frame_v2_t& InVideoFr
 			, InTime
 			, FFrameRate(InVideoFrame.frame_rate_N, InVideoFrame.frame_rate_D)
 			, InTimecode
-			, InColorFormatArgs);
+			, InSourceColorSettings);
 
 		default:
 			return false;
diff --git a/Engine/Plugins/Media/NDIMedia/Source/NDIMedia/Private/Player/NDIMediaTextureSample.h b/Engine/Plugins/Media/NDIMedia/Source/NDIMedia/Private/Player/NDIMediaTextureSample.h
--- a/Engine/Plugins/Media/NDIMedia/Source/NDIMedia/Private/Player/NDIMediaTextureSample.h
+++ b/Engine/Plugins/Media/NDIMedia/Source/NDIMedia/Private/Player/NDIMediaTextureSample.h
@@ -25,11 +25,11 @@ public:
 	 * Initialize the sample.
 	 * 
 	 * @param InVideoFrame Received Video frame data from NDI.
-	 * @param InColorFormatArgs Color space and encoding settings from the media source.
+	 * @param InSourceColorSettings Information about the texture color encoding and color space.
 	 * @param InTime Current playback time of the sample.
 	 * @return true if the sample was correctly initialized, false otherwise.
 	 */
-	bool Initialize(const NDIlib_video_frame_v2_t& InVideoFrame, const UE::MediaIOCore::FColorFormatArgs& InColorFormatArgs, const FTimespan& InTime, const TOptional<FTimecode>& InTimecode);
+	bool Initialize(const NDIlib_video_frame_v2_t& InVideoFrame, TSharedPtr<FNativeMediaSourceColorSettings, ESPMode::ThreadSafe> InSourceColorSettings, const FTimespan& InTime, const TOptional<FTimecode>& InTimecode);
 
 	/** Progressive vs Interlaced. */
 	bool bIsProgressive = true;
diff --git a/Engine/Plugins/Media/NDIMedia/Source/NDIMedia/Public/NDIMediaSource.h b/Engine/Plugins/Media/NDIMedia/Source/NDIMedia/Public/NDIMediaSource.h
--- a/Engine/Plugins/Media/NDIMedia/Source/NDIMedia/Public/NDIMediaSource.h
+++ b/Engine/Plugins/Media/NDIMedia/Source/NDIMedia/Public/NDIMediaSource.h
@@ -12,7 +12,7 @@
 /**
  * Media source for NDI streams.
  */
-UCLASS(MinimalAPI, BlueprintType, hideCategories=(Platforms,Object), PrioritizeCategories=(NDI, Video, Audio, Ancillary, Synchronization, Debug), meta=(MediaIOCustomLayout="NDI"))
+UCLASS(MinimalAPI, BlueprintType, hideCategories=(Platforms,Object), PrioritizeCategories=(NDI, Video, ColorManagement, Audio, Ancillary, Synchronization, Debug), meta=(MediaIOCustomLayout="NDI"))
 class UNDIMediaSource : public UCaptureCardMediaSource
 {
 	GENERATED_BODY()
diff --git a/Engine/Plugins/VirtualProduction/Rivermax/RivermaxMedia/Source/RivermaxMedia/Private/RivermaxMediaPlayer.cpp b/Engine/Plugins/VirtualProduction/Rivermax/RivermaxMedia/Source/RivermaxMedia/Private/RivermaxMediaPlayer.cpp
--- a/Engine/Plugins/VirtualProduction/Rivermax/RivermaxMedia/Source/RivermaxMedia/Private/RivermaxMediaPlayer.cpp
+++ b/Engine/Plugins/VirtualProduction/Rivermax/RivermaxMedia/Source/RivermaxMedia/Private/RivermaxMediaPlayer.cpp
@@ -504,11 +504,7 @@ namespace UE::RivermaxMedia
 
 				ProxySampleDummy = MakeShared<FRivermaxMediaTextureSample>();
 				ProxySampleDummy->SetTexture(CreateIntermediateRenderTarget(RHICmdList, Resolution, NewSample->GetPixelFormat(), NewSample->IsOutputSrgb()));
-				UE::MediaIOCore::FColorFormatArgs ColorArgs;
-				ColorArgs.ColorSpaceType = bOverrideSourceColorSpace ? (UE::Color::EColorSpace)OverrideSourceColorSpace : UE::Color::EColorSpace::None;
-				ColorArgs.Encoding = bOverrideSourceEncoding ? (UE::Color::EEncoding)OverrideSourceEncoding : UE::Color::EEncoding::None;
-
-				ProxySampleDummy->SetProperties(NewSample->GetStride(), VideoTrackFormat.Dim.X, VideoTrackFormat.Dim.X, NewSample->GetFormat(), FTimespan(0), FFrameRate(), FTimecode(), ColorArgs);
+				ProxySampleDummy->SetProperties(NewSample->GetStride(), VideoTrackFormat.Dim.X, VideoTrackFormat.Dim.Y, NewSample->GetFormat(), FTimespan(0), FFrameRate(), FTimecode(), SourceColorSettings);
 			});
 
 
diff --git a/Engine/Plugins/VirtualProduction/Rivermax/RivermaxMedia/Source/RivermaxMedia/Private/RivermaxMediaSource.cpp b/Engine/Plugins/VirtualProduction/Rivermax/RivermaxMedia/Source/RivermaxMedia/Private/RivermaxMediaSource.cpp
--- a/Engine/Plugins/VirtualProduction/Rivermax/RivermaxMedia/Source/RivermaxMedia/Private/RivermaxMediaSource.cpp
+++ b/Engine/Plugins/VirtualProduction/Rivermax/RivermaxMedia/Source/RivermaxMedia/Private/RivermaxMediaSource.cpp
@@ -14,8 +14,6 @@
 URivermaxMediaSource::URivermaxMediaSource() : UCaptureCardMediaSource()
 {
 	Deinterlacer = nullptr;
-	bOverrideSourceEncoding = false;
-	bOverrideSourceColorSpace = false;
 	bRenderJIT = true;
 	EvaluationType = EMediaIOSampleEvaluationType::Latest;
 }
@@ -126,8 +124,7 @@ void URivermaxMediaSource::PostLoad()
 
 		if (bIsSRGBInput_DEPRECATED == true)
 		{
-			bOverrideSourceEncoding = true;
-			OverrideSourceEncoding = EMediaIOCoreSourceEncoding::sRGB;
+			SourceColorSettings.EncodingOverride = EMediaSourceEncoding::MSE_sRGB;
 		}
 
 		if (PlayerMode_DEPRECATED == ERivermaxPlayerMode_DEPRECATED::Framelock)
diff --git a/Engine/Plugins/VirtualProduction/Rivermax/RivermaxMedia/Source/RivermaxMedia/Private/RivermaxMediaTextureSampleConverter.cpp b/Engine/Plugins/VirtualProduction/Rivermax/RivermaxMedia/Source/RivermaxMedia/Private/RivermaxMediaTextureSampleConverter.cpp
--- a/Engine/Plugins/VirtualProduction/Rivermax/RivermaxMedia/Source/RivermaxMedia/Private/RivermaxMediaTextureSampleConverter.cpp
+++ b/Engine/Plugins/VirtualProduction/Rivermax/RivermaxMedia/Source/RivermaxMedia/Private/RivermaxMediaTextureSampleConverter.cpp
@@ -8,7 +8,7 @@ namespace UE::RivermaxMedia
 	{
 		if (TSharedPtr<FMediaIOCoreTextureSampleBase> ProxySampleSharedPtr = JITRProxySample.Pin())
 		{
-			if (ProxySampleSharedPtr->GetEncodingType() != UE::Color::EEncoding::None || ProxySampleSharedPtr->GetColorSpaceType() != UE::Color::EColorSpace::None)
+			if (ProxySampleSharedPtr->HasEncodingOverride() || ProxySampleSharedPtr->HasColorSpaceOverride())
 			{
 				return ConverterInfoFlags_WillCreateOutputTexture | ConverterInfoFlags_PreprocessOnly;
 			}
diff --git a/Engine/Source/Runtime/Core/Private/ColorManagement/ColorSpace.cpp b/Engine/Source/Runtime/Core/Private/ColorManagement/ColorSpace.cpp
--- a/Engine/Source/Runtime/Core/Private/ColorManagement/ColorSpace.cpp
+++ b/Engine/Source/Runtime/Core/Private/ColorManagement/ColorSpace.cpp
@@ -76,6 +76,13 @@ const FColorSpace& FColorSpace::GetSRGB()
 	return ColorSpaceSRGB;
 }
 
+const FColorSpace& FColorSpace::GetRec2020()
+{
+	static const FColorSpace ColorSpaceRec2020 = FColorSpace(EColorSpace::Rec2020);
+
+	return ColorSpaceRec2020;
+}
+
 const FColorSpace& FColorSpace::GetP3D65()
 {
 	static const FColorSpace ColorSpaceP3D65 = FColorSpace(EColorSpace::P3D65);
diff --git a/Engine/Source/Runtime/Core/Public/ColorManagement/ColorSpace.h b/Engine/Source/Runtime/Core/Public/ColorManagement/ColorSpace.h
--- a/Engine/Source/Runtime/Core/Public/ColorManagement/ColorSpace.h
+++ b/Engine/Source/Runtime/Core/Public/ColorManagement/ColorSpace.h
@@ -118,6 +118,13 @@ public:
 	 */
 	static CORE_API const FColorSpace& GetSRGB();
 
+	/**
+	 * Get a (statically cached) Rec2020 color space.
+	 *
+	 * @return FColorSpace Rec2020 color space
+	 */
+	static CORE_API const FColorSpace& GetRec2020();
+
 	/**
 	 * Get a (statically cached) P3-D65 color space.
 	 *
diff --git a/Engine/Source/Runtime/Core/Public/UObject/FortniteMainBranchObjectVersions.inl b/Engine/Source/Runtime/Core/Public/UObject/FortniteMainBranchObjectVersions.inl
--- a/Engine/Source/Runtime/Core/Public/UObject/FortniteMainBranchObjectVersions.inl
+++ b/Engine/Source/Runtime/Core/Public/UObject/FortniteMainBranchObjectVersions.inl
@@ -706,5 +706,8 @@ FORTNITE_MAIN_VERSION(UAFVariablesGuid, 230)
 // Introducing FGuid cached to FAnimNext(Soft)VariableReference
 FORTNITE_MAIN_VERSION(UAFVariableReferenceGUID, 231)
 
+// Unified source color settings between img media & media io
+FORTNITE_MAIN_VERSION(MediaColorSettingsUnification, 232)
+
 // -----<new versions can be added above this line>-------------------------------------------------
 #undef FORTNITE_MAIN_VERSION
diff --git a/Engine/Source/Runtime/MediaAssets/Private/BaseMediaSourceColorSettings.cpp b/Engine/Source/Runtime/MediaAssets/Private/BaseMediaSourceColorSettings.cpp
new file mode 100644
--- /dev/null
+++ b/Engine/Source/Runtime/MediaAssets/Private/BaseMediaSourceColorSettings.cpp
@@ -0,0 +1,100 @@
+// Copyright Epic Games, Inc. All Rights Reserved.
+
+#include "BaseMediaSourceColorSettings.h"
+
+#include UE_INLINE_GENERATED_CPP_BY_NAME(BaseMediaSourceColorSettings)
+
+#if WITH_EDITOR
+void FMediaSourceColorSettings::UpdateColorSpaceChromaticities()
+{
+	if (ColorSpaceOverride != ETextureColorSpace::TCS_Custom)
+	{
+		UE::Color::FColorSpace Chromaticities(static_cast<UE::Color::EColorSpace>(ColorSpaceOverride));
+		Chromaticities.GetChromaticities(RedChromaticityCoordinate, GreenChromaticityCoordinate, BlueChromaticityCoordinate, WhiteChromaticityCoordinate);
+	}
+}
+#endif // WITH_EDITOR
+
+FNativeMediaSourceColorSettings::FNativeMediaSourceColorSettings()
+	: EncodingOverride(UE::Color::EEncoding::None)
+	, ColorSpaceOverride()
+	, ChromaticAdaptationMethod(UE::Color::EChromaticAdaptationMethod::Bradford)
+{}
+
+FNativeMediaSourceColorSettings::~FNativeMediaSourceColorSettings() = default;
+
+FNativeMediaSourceColorSettings::FNativeMediaSourceColorSettings(const FNativeMediaSourceColorSettings& Other)
+{
+	EncodingOverride.store(Other.EncodingOverride);
+	ChromaticAdaptationMethod.store(Other.ChromaticAdaptationMethod);
+
+	FScopeLock Lock(&Other.ColorSpaceCriticalSection);
+	ColorSpaceOverride = Other.ColorSpaceOverride;
+}
+
+FNativeMediaSourceColorSettings& FNativeMediaSourceColorSettings::operator=(const FNativeMediaSourceColorSettings& Other)
+{
+	if (this != &Other)
+	{
+		EncodingOverride.store(Other.EncodingOverride);
+		ChromaticAdaptationMethod.store(Other.ChromaticAdaptationMethod);
+
+		FScopeLock Lock(&ColorSpaceCriticalSection);
+		FScopeLock LockOther(&Other.ColorSpaceCriticalSection);
+		ColorSpaceOverride = Other.ColorSpaceOverride;
+	}
+	return *this;
+}
+
+void FNativeMediaSourceColorSettings::Update(const FMediaSourceColorSettings& InSettings)
+{
+	EncodingOverride = static_cast<UE::Color::EEncoding>(InSettings.EncodingOverride);
+	ChromaticAdaptationMethod = static_cast<UE::Color::EChromaticAdaptationMethod>(InSettings.ChromaticAdaptationMethod);
+
+	FScopeLock Lock(&ColorSpaceCriticalSection);
+
+	if (InSettings.ColorSpaceOverride == ETextureColorSpace::TCS_Custom)
+	{
+		ColorSpaceOverride = UE::Color::FColorSpace(
+			InSettings.RedChromaticityCoordinate,
+			InSettings.GreenChromaticityCoordinate,
+			InSettings.BlueChromaticityCoordinate,
+			InSettings.WhiteChromaticityCoordinate
+		);
+	}
+	else if (InSettings.ColorSpaceOverride != ETextureColorSpace::TCS_None)
+	{
+		ColorSpaceOverride = UE::Color::FColorSpace(static_cast<UE::Color::EColorSpace>(InSettings.ColorSpaceOverride));
+	}
+	else
+	{
+		ColorSpaceOverride.Reset();
+	}
+}
+
+bool FNativeMediaSourceColorSettings::HasColorSpaceOverride() const
+{
+	FScopeLock Lock(&ColorSpaceCriticalSection);
+
+	return ColorSpaceOverride.IsSet();
+}
+
+const UE::Color::FColorSpace& FNativeMediaSourceColorSettings::GetColorSpaceOverride(const UE::Color::FColorSpace& InDefaultColorSpace) const
+{
+	FScopeLock Lock(&ColorSpaceCriticalSection);
+	
+	if (ColorSpaceOverride.IsSet())
+	{
+		return ColorSpaceOverride.GetValue();
+	}
+
+	return InDefaultColorSpace;
+}
+
+void FNativeMediaSourceColorSettings::SetColorSpaceOverride(const UE::Color::FColorSpace& InColorSpaceOverride)
+{
+	FScopeLock Lock(&ColorSpaceCriticalSection);
+
+	ColorSpaceOverride = InColorSpaceOverride;
+}
+
diff --git a/Engine/Source/Runtime/MediaAssets/Public/BaseMediaSourceColorSettings.h b/Engine/Source/Runtime/MediaAssets/Public/BaseMediaSourceColorSettings.h
new file mode 100644
--- /dev/null
+++ b/Engine/Source/Runtime/MediaAssets/Public/BaseMediaSourceColorSettings.h
@@ -0,0 +1,126 @@
+// Copyright Epic Games, Inc. All Rights Reserved.
+
+#pragma once
+
+#include "CoreTypes.h"
+#include "ColorManagement/ColorSpace.h"
+#include "Engine/TextureDefines.h"
+#include "IMediaOptions.h"
+
+#include "BaseMediaSourceColorSettings.generated.h"
+
+/** List of source encodings that can be converted to linear. (Integer values match the ETextureSourceEncoding values in TextureDefines.h */
+UENUM()
+enum class EMediaSourceEncoding : uint8
+{
+	MSE_None	= 0 UMETA(DisplayName = "None", ToolTip = "The source encoding is not overridden."),
+	MSE_Linear	= 1 UMETA(DisplayName = "Linear", ToolTip = "The source encoding is considered linear."),
+	MSE_sRGB	= 2 UMETA(DisplayName = "sRGB", ToolTip = "sRGB source encoding to be linearized"),
+	MSE_ST2084	= 3 UMETA(DisplayName = "ST 2084/PQ", ToolTip = "SMPTE ST 2084/PQ source encoding to be linearized"),
+	MSE_SLog3	= 12 UMETA(DisplayName = "SLog3", ToolTip = "Sony SLog3 source encoding to be linearized"),
+
+	MSE_MAX,
+};
+
+/* Manual definition of media source color space & encoding. */
+USTRUCT(BlueprintType)
+struct FMediaSourceColorSettings
+{
+	GENERATED_USTRUCT_BODY()
+
+	FMediaSourceColorSettings()
+		: EncodingOverride(EMediaSourceEncoding::MSE_None)
+		, ColorSpaceOverride(ETextureColorSpace::TCS_None)
+		, RedChromaticityCoordinate(FVector2D::ZeroVector)
+		, GreenChromaticityCoordinate(FVector2D::ZeroVector)
+		, BlueChromaticityCoordinate(FVector2D::ZeroVector)
+		, WhiteChromaticityCoordinate(FVector2D::ZeroVector)
+		, ChromaticAdaptationMethod(static_cast<ETextureChromaticAdaptationMethod>(UE::Color::DEFAULT_CHROMATIC_ADAPTATION_METHOD))
+	{}
+
+	/** Source encoding of the media. */
+	UPROPERTY(EditAnywhere, BlueprintReadWrite, Category = ColorManagement)
+	EMediaSourceEncoding EncodingOverride;
+
+	/** Source color space of the media. */
+	UPROPERTY(EditAnywhere, BlueprintReadWrite, Category = ColorManagement)
+	ETextureColorSpace ColorSpaceOverride;
+
+	/** Red chromaticity coordinate of the source color space. */
+	UPROPERTY(EditAnywhere, BlueprintReadWrite, Category = ColorManagement, meta = (EditCondition = "ColorSpaceOverride == ETextureColorSpace::TCS_Custom", EditConditionHides))
+	FVector2D RedChromaticityCoordinate;
+
+	/** Green chromaticity coordinate of the source color space. */
+	UPROPERTY(EditAnywhere, BlueprintReadWrite, Category = ColorManagement, meta = (EditCondition = "ColorSpaceOverride == ETextureColorSpace::TCS_Custom", EditConditionHides))
+	FVector2D GreenChromaticityCoordinate;
+
+	/** Blue chromaticity coordinate of the source color space. */
+	UPROPERTY(EditAnywhere, BlueprintReadWrite, Category = ColorManagement, meta = (EditCondition = "ColorSpaceOverride == ETextureColorSpace::TCS_Custom", EditConditionHides))
+	FVector2D BlueChromaticityCoordinate;
+
+	/** White chromaticity coordinate of the source color space. */
+	UPROPERTY(EditAnywhere, BlueprintReadWrite, Category = ColorManagement, meta = (EditCondition = "ColorSpaceOverride == ETextureColorSpace::TCS_Custom", EditConditionHides))
+	FVector2D WhiteChromaticityCoordinate;
+
+	/** Chromatic adaption method applied if the source white point differs from the working color space white point. */
+	UPROPERTY(EditAnywhere, BlueprintReadWrite, Category = ColorManagement, meta = (EditCondition = "ColorSpaceOverride != ETextureColorSpace::TCS_None"))
+	ETextureChromaticAdaptationMethod ChromaticAdaptationMethod;
+
+#if WITH_EDITOR
+	/** Update the chromaticity coordinates member variables based on the color space choice (unless custom). */
+	MEDIAASSETS_API void UpdateColorSpaceChromaticities();
+#endif
+};
+
+/* Engine-native color source settings container for media option. */
+struct FNativeMediaSourceColorSettings : public IMediaOptions::FDataContainer
+{
+	/** Constructor */
+	MEDIAASSETS_API FNativeMediaSourceColorSettings();
+
+	/** Destructor */
+	MEDIAASSETS_API ~FNativeMediaSourceColorSettings();
+
+	/** Copy constructor */
+	MEDIAASSETS_API FNativeMediaSourceColorSettings(const FNativeMediaSourceColorSettings& Other);
+
+	/** Assignment operator */
+	MEDIAASSETS_API FNativeMediaSourceColorSettings& operator=(const FNativeMediaSourceColorSettings& Other);
+
+	/** Updates the native settings from user-controlled settings. */
+	MEDIAASSETS_API void Update(const FMediaSourceColorSettings& InSettings);
+
+	/** Has encoding override. */
+	inline bool HasEncodingOverride() const { return EncodingOverride != UE::Color::EEncoding::None; }
+
+	/** Color encoding override getter. */
+	inline UE::Color::EEncoding GetEncodingOverride() const { return EncodingOverride; }
+
+	/** Color encoding override setter. */
+	inline void SetEncodingOverride(UE::Color::EEncoding InEncodingOverride) { EncodingOverride = InEncodingOverride; }
+
+	/** Has color space override. */
+	MEDIAASSETS_API bool HasColorSpaceOverride() const;
+
+	/** Color space override getter. */
+	MEDIAASSETS_API const UE::Color::FColorSpace& GetColorSpaceOverride(const UE::Color::FColorSpace& InDefaultColorSpace) const;
+
+	/** Color space override setter. */
+	MEDIAASSETS_API void SetColorSpaceOverride(const UE::Color::FColorSpace& InColorSpaceOverride);
+
+	/** Chromatic adaptation getter. */
+	inline UE::Color::EChromaticAdaptationMethod GetChromaticAdaptationMethod() const { return ChromaticAdaptationMethod; }
+
+private:
+	/** Manual source encoding override. */
+	std::atomic<UE::Color::EEncoding> EncodingOverride;
+
+	/** Manual source color space override. */
+	TOptional<UE::Color::FColorSpace> ColorSpaceOverride;
+
+	/** Chromatic adapation to be used on manual source color space override. */
+	std::atomic<UE::Color::EChromaticAdaptationMethod> ChromaticAdaptationMethod;
+
+	/** Protects color space override variable. */
+	mutable FCriticalSection ColorSpaceCriticalSection;
+};
