MediaIO: CaptureCard CIS fix with duplicate script names.

diff --git a/Engine/Plugins/Media/MediaIOFramework/Source/MediaIOCore/Public/CaptureCardMediaSource.h b/Engine/Plugins/Media/MediaIOFramework/Source/MediaIOCore/Public/CaptureCardMediaSource.h
--- a/Engine/Plugins/Media/MediaIOFramework/Source/MediaIOCore/Public/CaptureCardMediaSource.h
+++ b/Engine/Plugins/Media/MediaIOFramework/Source/MediaIOCore/Public/CaptureCardMediaSource.h
@@ -134,7 +134,7 @@ public:
 #if WITH_EDITORONLY_DATA
 PRAGMA_DISABLE_DEPRECATION_WARNINGS
 	UE_DEPRECATED(5.8, "This property is deprecated. Please use SourceColorSettings instead.")
-	UPROPERTY(meta = (DeprecatedProperty, DeprecationMessage = "This property is deprecated. Please use SourceColorSettings instead."))
+	UPROPERTY(meta = (ScriptName = "IsOverrideSourceEncoding", DeprecatedProperty, DeprecationMessage = "This property is deprecated. Please use SourceColorSettings instead."))
 	bool bOverrideSourceEncoding_DEPRECATED = true;
 
 	UE_DEPRECATED(5.8, "This property is deprecated. Please use SourceColorSettings instead.")
@@ -142,7 +142,7 @@ PRAGMA_DISABLE_DEPRECATION_WARNINGS
 	EMediaIOCoreSourceEncoding OverrideSourceEncoding_DEPRECATED = EMediaIOCoreSourceEncoding::Linear;
 
 	UE_DEPRECATED(5.8, "This property is deprecated. Please use SourceColorSettings instead.")
-	UPROPERTY(meta = (DeprecatedProperty, DeprecationMessage = "This property is deprecated. Please use SourceColorSettings instead."))
+	UPROPERTY(meta = (ScriptName = "IsOverrideSourceColorSpace", DeprecatedProperty, DeprecationMessage = "This property is deprecated. Please use SourceColorSettings instead."))
 	bool bOverrideSourceColorSpace_DEPRECATED = true;
 
 	UE_DEPRECATED(5.8, "This property is deprecated. Please use SourceColorSettings instead.")
