Media IO - Fix assertion when creating media source while running with the Vulkan renderer

#rnx

#Synced-CL 48456341

diff --git a/Engine/Plugins/Media/MediaIOFramework/MediaIOFramework.uplugin b/Engine/Plugins/Media/MediaIOFramework/MediaIOFramework.uplugin
--- a/Engine/Plugins/Media/MediaIOFramework/MediaIOFramework.uplugin
+++ b/Engine/Plugins/Media/MediaIOFramework/MediaIOFramework.uplugin
@@ -28,7 +28,7 @@
 		{
 			"Name": "GPUTextureTransfer",
 			"Type": "Runtime",
-			"LoadingPhase": "Default",
+			"LoadingPhase": "EarliestPossible",
 			"PlatformArchitectureDenyList": [ "Win64:arm64" ]
 		}
 	],
diff --git a/Engine/Plugins/Media/MediaIOFramework/Source/GPUTextureTransfer/Private/GPUTextureTransferModule.cpp b/Engine/Plugins/Media/MediaIOFramework/Source/GPUTextureTransfer/Private/GPUTextureTransferModule.cpp
--- a/Engine/Plugins/Media/MediaIOFramework/Source/GPUTextureTransfer/Private/GPUTextureTransferModule.cpp
+++ b/Engine/Plugins/Media/MediaIOFramework/Source/GPUTextureTransfer/Private/GPUTextureTransferModule.cpp
@@ -69,8 +69,26 @@ void FGPUTextureTransferModule::StartupModule()
 
 	CVarMediaIOEnableGPUDirect.AsVariable()->OnChangedDelegate().AddRaw(this, &FGPUTextureTransferModule::OnEnableGPUDirectCVarChange);
 
-	// Cache this information since GetGPUDriverInfo has to be called on game thread because of a call to GetValueOnGameThread.
-	CachedDriverInfo = FPlatformMisc::GetGPUDriverInfo(GRHIAdapterName);
+	if (FApp::CanEverRender())
+	{
+		// Always provide the necessary Vulkan extensions (it will just get ignored if a different RHI is used)
+#if DVP_SUPPORTED_PLATFORM
+		const TArray<const ANSICHAR*> ExtentionsToAdd{ VK_KHR_EXTERNAL_MEMORY_EXTENSION_NAME, VK_KHR_EXTERNAL_MEMORY_WIN32_EXTENSION_NAME, VK_KHR_EXTERNAL_SEMAPHORE_WIN32_EXTENSION_NAME, VK_KHR_SURFACE_EXTENSION_NAME, VK_KHR_WIN32_SURFACE_EXTENSION_NAME };
+#elif PLATFORM_LINUX
+		const TArray<const ANSICHAR*> ExtentionsToAdd{ VK_KHR_EXTERNAL_MEMORY_EXTENSION_NAME, VK_KHR_EXTERNAL_MEMORY_FD_EXTENSION_NAME, VK_KHR_EXTERNAL_SEMAPHORE_FD_EXTENSION_NAME, VK_KHR_SURFACE_EXTENSION_NAME };
+#endif
+
+#if VULKAN_PLATFORM
+		IVulkanDynamicRHI::AddEnabledDeviceExtensionsAndLayers(ExtentionsToAdd, TArray<const ANSICHAR*>());
+#endif
+	}
+
+	// Defer the following since the module is loaded too early for this call to succeed.
+	FCoreDelegates::OnPostEngineInit.AddLambda([this]() 
+	{
+		// Cache this information since GetGPUDriverInfo has to be called on game thread because of a call to GetValueOnGameThread.
+		CachedDriverInfo = FPlatformMisc::GetGPUDriverInfo(GRHIAdapterName);
+	});
 }
 
 void FGPUTextureTransferModule::ShutdownModule()
@@ -87,19 +105,6 @@ void FGPUTextureTransferModule::Initialize()
 	{
 		if (LoadGPUDirectBinary())
 		{
-			// Always provide the necessary Vulkan extensions (it will just get ignored if a different RHI is used)
-			{
-#if DVP_SUPPORTED_PLATFORM
-				const TArray<const ANSICHAR*> ExtentionsToAdd{ VK_KHR_EXTERNAL_MEMORY_EXTENSION_NAME, VK_KHR_EXTERNAL_MEMORY_WIN32_EXTENSION_NAME, VK_KHR_EXTERNAL_SEMAPHORE_WIN32_EXTENSION_NAME, VK_KHR_SURFACE_EXTENSION_NAME, VK_KHR_WIN32_SURFACE_EXTENSION_NAME };
-#elif PLATFORM_LINUX
-				const TArray<const ANSICHAR*> ExtentionsToAdd{ VK_KHR_EXTERNAL_MEMORY_EXTENSION_NAME, VK_KHR_EXTERNAL_MEMORY_FD_EXTENSION_NAME, VK_KHR_EXTERNAL_SEMAPHORE_FD_EXTENSION_NAME, VK_KHR_SURFACE_EXTENSION_NAME };
-#endif
-
-#if VULKAN_PLATFORM
-				IVulkanDynamicRHI::AddEnabledDeviceExtensionsAndLayers(ExtentionsToAdd, TArray<const ANSICHAR*>());
-#endif
-			}
-
 			TransferObjects.AddDefaulted(RHI_MAX);
 
 			InitializeTextureTransfer();
