[ContentBrowser] New Add Menu Media Menu

#Synced-CL 52619962

diff --git a/Engine/Plugins/Experimental/ImagePlate/Source/ImagePlateEditor/Private/AssetDefinition_ImagePlateFileSequence.h b/Engine/Plugins/Experimental/ImagePlate/Source/ImagePlateEditor/Private/AssetDefinition_ImagePlateFileSequence.h
--- a/Engine/Plugins/Experimental/ImagePlate/Source/ImagePlateEditor/Private/AssetDefinition_ImagePlateFileSequence.h
+++ b/Engine/Plugins/Experimental/ImagePlate/Source/ImagePlateEditor/Private/AssetDefinition_ImagePlateFileSequence.h
@@ -19,7 +19,7 @@ public:
 	virtual TSoftClassPtr<UObject> GetAssetClass() const override { return UImagePlateFileSequence::StaticClass(); }
 	virtual TConstArrayView<FAssetCategoryPath> GetAssetCategories() const override
 	{
-		static const auto Categories = { EAssetCategoryPaths::Media };
+		static const auto Categories = { FAssetCategoryPath(EAssetCategoryPaths::Media, NSLOCTEXT("AssetDefinition", "AssetDefinition_ImagePlateFileSequenceSubMenu", "Other"), ECategoryMenuType::Section) };
 		return Categories;
 	}
 	// UAssetDefinition End
diff --git a/Engine/Plugins/Media/ImgMedia/Source/ImgMediaEditor/ImgMediaEditor.Build.cs b/Engine/Plugins/Media/ImgMedia/Source/ImgMediaEditor/ImgMediaEditor.Build.cs
--- a/Engine/Plugins/Media/ImgMedia/Source/ImgMediaEditor/ImgMediaEditor.Build.cs
+++ b/Engine/Plugins/Media/ImgMedia/Source/ImgMediaEditor/ImgMediaEditor.Build.cs
@@ -8,7 +8,7 @@ namespace UnrealBuildTool.Rules
 		{
 			PrivateDependencyModuleNames.AddRange(
 				new string[] {
-					"AssetTools",
+					"AssetDefinition",
 					"ContentBrowser",
 					"Core",
 					"CoreUObject",
diff --git a/Engine/Plugins/Media/ImgMedia/Source/ImgMediaEditor/Private/AssetDefinition/AssetDefinition_ImgMediaSource.h b/Engine/Plugins/Media/ImgMedia/Source/ImgMediaEditor/Private/AssetDefinition/AssetDefinition_ImgMediaSource.h
new file mode 100644
--- /dev/null
+++ b/Engine/Plugins/Media/ImgMedia/Source/ImgMediaEditor/Private/AssetDefinition/AssetDefinition_ImgMediaSource.h
@@ -0,0 +1,22 @@
+// Copyright Epic Games, Inc. All Rights Reserved.
+
+#pragma once
+#include "ImgMediaSource.h"
+#include "AssetDefinition/AssetDefinition_MediaSource.h"
+#include "AssetDefinition_ImgMediaSource.generated.h"
+
+UCLASS()
+class UAssetDefinition_ImgMediaSource : public UAssetDefinition_MediaSource
+{
+	GENERATED_BODY()
+public:
+	// UAssetDefinition Begin
+	virtual FText GetAssetDisplayName() const override { return NSLOCTEXT("AssetDefinition", "AssetDefinition_ImgMediaSource", "Img Media Source"); }
+	virtual TSoftClassPtr<UObject> GetAssetClass() const override { return UImgMediaSource::StaticClass(); }
+	virtual TConstArrayView<FAssetCategoryPath> GetAssetCategories() const override
+	{
+		static const auto Categories = { FAssetCategoryPath(EAssetCategoryPaths::Media, NSLOCTEXT("AssetDefinition", "AssetDefinition_ImgMediaSourceSubMenu", "Media Source"), ECategoryMenuType::Section) };
+		return Categories;
+	}
+	// UAssetDefinition End
+};
diff --git a/Engine/Plugins/Media/ImgMedia/Source/ImgMediaEditor/Private/AssetTools/ImgMediaSourceActions.cpp b/Engine/Plugins/Media/ImgMedia/Source/ImgMediaEditor/Private/AssetTools/ImgMediaSourceActions.cpp
deleted file mode 100644
diff --git a/Engine/Plugins/Media/ImgMedia/Source/ImgMediaEditor/Private/AssetTools/ImgMediaSourceActions.h b/Engine/Plugins/Media/ImgMedia/Source/ImgMediaEditor/Private/AssetTools/ImgMediaSourceActions.h
deleted file mode 100644
diff --git a/Engine/Plugins/Media/ImgMedia/Source/ImgMediaEditor/Private/ImgMediaEditorModule.cpp b/Engine/Plugins/Media/ImgMedia/Source/ImgMediaEditor/Private/ImgMediaEditorModule.cpp
--- a/Engine/Plugins/Media/ImgMedia/Source/ImgMediaEditor/Private/ImgMediaEditorModule.cpp
+++ b/Engine/Plugins/Media/ImgMedia/Source/ImgMediaEditor/Private/ImgMediaEditorModule.cpp
@@ -3,12 +3,9 @@
 #include "ImgMediaEditorModule.h"
 #include "Modules/ModuleManager.h"
 
-#include "AssetToolsModule.h"
-#include "AssetTools/ImgMediaSourceActions.h"
 #include "ContentBrowserMenuContexts.h"
 #include "Customizations/ImgMediaSourceCustomization.h"
 #include "Framework/Application/SlateApplication.h"
-#include "IAssetTools.h"
 #include "ImgMediaEditorModule.h"
 #include "IImgMediaModule.h"
 #include "ImgMediaSource.h"
@@ -49,7 +46,6 @@ public:
 	virtual void StartupModule() override
 	{
 		RegisterCustomizations();
-		RegisterAssetTools();
 		RegisterTabSpawners();
 		ExtendContentMenu();
 
@@ -64,7 +60,6 @@ public:
 	virtual void ShutdownModule() override
 	{
 		UnregisterTabSpawners();
-		UnregisterAssetTools();
 		UnregisterCustomizations();
 	}
 
@@ -94,30 +89,6 @@ protected:
 		}
 	}
 
-	void RegisterAssetTools()
-	{
-		IAssetTools& AssetTools = FModuleManager::LoadModuleChecked<FAssetToolsModule>("AssetTools").Get();
-
-		TSharedRef<IAssetTypeActions> Action = MakeShared<FImgMediaSourceActions>();
-		AssetTools.RegisterAssetTypeActions(Action);
-		RegisteredAssetTypeActions.Add(Action);
-	}
-
-	void UnregisterAssetTools()
-	{
-		FAssetToolsModule* AssetToolsModule = FModuleManager::GetModulePtr<FAssetToolsModule>("AssetTools");
-
-		if (AssetToolsModule != nullptr)
-		{
-			IAssetTools& AssetTools = AssetToolsModule->Get();
-
-			for (const TSharedRef<IAssetTypeActions>& Action : RegisteredAssetTypeActions)
-			{
-				AssetTools.UnregisterAssetTypeActions(Action);
-			}
-		}
-	}
-
 	void RegisterTabSpawners()
 	{
 		// Add ImgMedia group.
@@ -294,9 +265,6 @@ private:
 	/** This will be passed to ProceessEXR and then cleared. */
 	static FString ProcessEXRInputPath;
 
-	/** The collection of registered asset type actions. */
-	TArray<TSharedRef<IAssetTypeActions>> RegisteredAssetTypeActions;
-
 	/** Array of all our players. */
 	TArray<TWeakPtr<FImgMediaPlayer>> MediaPlayers;
 };
diff --git a/Engine/Plugins/Media/MediaIOFramework/Source/MediaIOEditor/MediaIOEditor.Build.cs b/Engine/Plugins/Media/MediaIOFramework/Source/MediaIOEditor/MediaIOEditor.Build.cs
--- a/Engine/Plugins/Media/MediaIOFramework/Source/MediaIOEditor/MediaIOEditor.Build.cs
+++ b/Engine/Plugins/Media/MediaIOFramework/Source/MediaIOEditor/MediaIOEditor.Build.cs
@@ -22,6 +22,7 @@ namespace UnrealBuildTool.Rules
 			PrivateDependencyModuleNames.AddRange(
 				new string[]
 				{
+					"AssetDefinition",
 					"Engine",
 					"MediaAssets"
 				});
diff --git a/Engine/Plugins/Media/MediaIOFramework/Source/MediaIOEditor/Private/AssetDefinition/AssetDefinition_FileMediaOutput.h b/Engine/Plugins/Media/MediaIOFramework/Source/MediaIOEditor/Private/AssetDefinition/AssetDefinition_FileMediaOutput.h
new file mode 100644
--- /dev/null
+++ b/Engine/Plugins/Media/MediaIOFramework/Source/MediaIOEditor/Private/AssetDefinition/AssetDefinition_FileMediaOutput.h
@@ -0,0 +1,22 @@
+// Copyright Epic Games, Inc. All Rights Reserved.
+
+#pragma once
+#include "FileMediaOutput.h"
+#include "AssetDefinition/AssetDefinition_MediaOutput.h"
+#include "AssetDefinition_FileMediaOutput.generated.h"
+
+UCLASS()
+class UAssetDefinition_FileMediaOutput : public UAssetDefinition_MediaOutput
+{
+	GENERATED_BODY()
+public:
+	// UAssetDefinition Begin
+	virtual FText GetAssetDisplayName() const override { return NSLOCTEXT("AssetDefinition", "AssetDefinition_FileMediaOutput", "File Media Output"); }
+	virtual TSoftClassPtr<UObject> GetAssetClass() const override { return UFileMediaOutput::StaticClass(); }
+	virtual TConstArrayView<FAssetCategoryPath> GetAssetCategories() const override
+	{
+		static const auto Categories = { FAssetCategoryPath(EAssetCategoryPaths::Media, NSLOCTEXT("AssetDefinition", "AssetDefinition_FileMediaOutputSubMenu", "Media Output"), ECategoryMenuType::Section) };
+		return Categories;
+	}
+	// UAssetDefinition End
+};
diff --git a/Engine/Plugins/Media/MediaIOFramework/Source/MediaIOEditor/Private/AssetTypeActions/AssetTypeActions_MediaOutput.cpp b/Engine/Plugins/Media/MediaIOFramework/Source/MediaIOEditor/Private/AssetTypeActions/AssetTypeActions_MediaOutput.cpp
deleted file mode 100644
diff --git a/Engine/Plugins/Media/MediaIOFramework/Source/MediaIOEditor/Private/AssetTypeActions/AssetTypeActions_MediaOutput.h b/Engine/Plugins/Media/MediaIOFramework/Source/MediaIOEditor/Private/AssetTypeActions/AssetTypeActions_MediaOutput.h
deleted file mode 100644
diff --git a/Engine/Plugins/Media/MediaIOFramework/Source/MediaIOEditor/Private/FileMediaOutputFactory.cpp b/Engine/Plugins/Media/MediaIOFramework/Source/MediaIOEditor/Private/FileMediaOutputFactory.cpp
--- a/Engine/Plugins/Media/MediaIOFramework/Source/MediaIOEditor/Private/FileMediaOutputFactory.cpp
+++ b/Engine/Plugins/Media/MediaIOFramework/Source/MediaIOEditor/Private/FileMediaOutputFactory.cpp
@@ -26,16 +26,3 @@ UObject* UFileMediaOutputFactory::FactoryCreateNew(UClass* InClass, UObject* InP
 {
 	return NewObject<UFileMediaOutput>(InParent, InClass, InName, Flags);
 }
-
-
-uint32 UFileMediaOutputFactory::GetMenuCategories() const
-{
-	return EAssetTypeCategories::Media;
-}
-
-
-bool UFileMediaOutputFactory::ShouldShowInNewMenu() const
-{
-	return true;
-}
-
diff --git a/Engine/Plugins/Media/MediaIOFramework/Source/MediaIOEditor/Private/FileMediaOutputFactory.h b/Engine/Plugins/Media/MediaIOFramework/Source/MediaIOEditor/Private/FileMediaOutputFactory.h
--- a/Engine/Plugins/Media/MediaIOFramework/Source/MediaIOEditor/Private/FileMediaOutputFactory.h
+++ b/Engine/Plugins/Media/MediaIOFramework/Source/MediaIOEditor/Private/FileMediaOutputFactory.h
@@ -18,8 +18,6 @@ class UFileMediaOutputFactory : public UFactory
 
 public:
 	UE_API virtual UObject* FactoryCreateNew(UClass* InClass, UObject* InParent, FName InName, EObjectFlags Flags, UObject* Context, FFeedbackContext* Warn) override;
-	UE_API virtual uint32 GetMenuCategories() const override;
-	UE_API virtual bool ShouldShowInNewMenu() const override;
 };
 
 #undef UE_API
diff --git a/Engine/Plugins/Media/MediaIOFramework/Source/MediaIOEditor/Private/MediaIOEditorModule.cpp b/Engine/Plugins/Media/MediaIOFramework/Source/MediaIOEditor/Private/MediaIOEditorModule.cpp
--- a/Engine/Plugins/Media/MediaIOFramework/Source/MediaIOEditor/Private/MediaIOEditorModule.cpp
+++ b/Engine/Plugins/Media/MediaIOFramework/Source/MediaIOEditor/Private/MediaIOEditorModule.cpp
@@ -12,8 +12,6 @@
 #include "Styling/SlateStyle.h"
 #include "Styling/SlateStyleRegistry.h"
 
-#include "AssetTypeActions/AssetTypeActions_MediaOutput.h"
-
 #include "Customizations/MediaIODeviceCustomization.h"
 #include "Customizations/MediaIOConfigurationCustomization.h"
 #include "Customizations/MediaIOInputConfigurationCustomization.h"
@@ -29,7 +27,6 @@ class FMediaIOEditorModule : public IModuleInterface
 
 	virtual void StartupModule() override
 	{
-		RegisterAssetTools();
 		RegisterCustomizations();
 		RegisterStyle();
 	}
@@ -40,45 +37,10 @@ class FMediaIOEditorModule : public IModuleInterface
 		{
 			UnregisterStyle();
 			UnregisterCustomizations();
-			UnregisterAssetTools();
 		}
 	}
 
 private:
-	/** Registers asset tool actions. */
-	void RegisterAssetTools()
-	{
-		IAssetTools& AssetTools = FModuleManager::LoadModuleChecked<FAssetToolsModule>("AssetTools").Get();
-		RegisterAssetTypeAction(AssetTools, MakeShareable(new FAssetTypeActions_MediaOutput));
-	}
-
-	/**
-	 * Registers a single asset type action.
-	 */
-	void RegisterAssetTypeAction(IAssetTools& AssetTools, TSharedRef<IAssetTypeActions> Action)
-	{
-		AssetTools.RegisterAssetTypeActions(Action);
-		RegisteredAssetTypeActions.Add(Action);
-	}
-
-	/** Unregisters asset tool actions. */
-	void UnregisterAssetTools()
-	{
-		FAssetToolsModule* AssetToolsModule = FModuleManager::GetModulePtr<FAssetToolsModule>("AssetTools");
-
-		if (AssetToolsModule != nullptr)
-		{
-			IAssetTools& AssetTools = AssetToolsModule->Get();
-
-			for (const TSharedRef<IAssetTypeActions>& Action : RegisteredAssetTypeActions)
-			{
-				AssetTools.UnregisterAssetTypeActions(Action);
-			}
-		}
-
-		RegisteredAssetTypeActions.Reset();
-	}
-
 	/** Register details view customizations. */
 	void RegisterCustomizations()
 	{
@@ -122,10 +84,6 @@ private:
 	{
 		FSlateStyleRegistry::UnRegisterSlateStyle(*StyleInstance.Get());
 	}
-
-private:
-	/** The collection of registered asset type actions. */
-	TArray<TSharedRef<IAssetTypeActions>> RegisteredAssetTypeActions;
 };
 
 IMPLEMENT_MODULE(FMediaIOEditorModule, MediaIOEditor)
diff --git a/Engine/Plugins/Media/MediaIOFramework/Source/MediaIOEditor/Public/AssetDefinition/AssetDefinition_MediaOutput.h b/Engine/Plugins/Media/MediaIOFramework/Source/MediaIOEditor/Public/AssetDefinition/AssetDefinition_MediaOutput.h
new file mode 100644
--- /dev/null
+++ b/Engine/Plugins/Media/MediaIOFramework/Source/MediaIOEditor/Public/AssetDefinition/AssetDefinition_MediaOutput.h
@@ -0,0 +1,18 @@
+// Copyright Epic Games, Inc. All Rights Reserved.
+
+#pragma once
+#include "AssetDefinitionDefault.h"
+#include "MediaOutput.h"
+#include "AssetDefinition_MediaOutput.generated.h"
+
+UCLASS(Abstract)
+class UAssetDefinition_MediaOutput : public UAssetDefinitionDefault
+{
+	GENERATED_BODY()
+public:
+	// UAssetDefinition Begin
+	virtual FLinearColor GetAssetColor() const override { return FColor::White; }
+	virtual TSoftClassPtr<UObject> GetAssetClass() const override { return UMediaOutput::StaticClass(); }
+	virtual bool CanImport() const override { return false; }
+	// UAssetDefinition End
+};
diff --git a/Engine/Plugins/Media/MediaPlayerEditor/Source/MediaPlayerEditor/MediaPlayerEditor.Build.cs b/Engine/Plugins/Media/MediaPlayerEditor/Source/MediaPlayerEditor/MediaPlayerEditor.Build.cs
--- a/Engine/Plugins/Media/MediaPlayerEditor/Source/MediaPlayerEditor/MediaPlayerEditor.Build.cs
+++ b/Engine/Plugins/Media/MediaPlayerEditor/Source/MediaPlayerEditor/MediaPlayerEditor.Build.cs
@@ -17,9 +17,10 @@ public class MediaPlayerEditor : ModuleRules
 
 		PublicDependencyModuleNames.AddRange(
 			new string[] {
+				"AssetDefinition",
 				"Core",
 				"Slate",
-				"SlateCore",
+				"SlateCore"
 			});
 
 		PrivateDependencyModuleNames.AddRange(
diff --git a/Engine/Plugins/Media/MediaPlayerEditor/Source/MediaPlayerEditor/Private/AssetDefinition/AssetDefinition_FileMediaSource.h b/Engine/Plugins/Media/MediaPlayerEditor/Source/MediaPlayerEditor/Private/AssetDefinition/AssetDefinition_FileMediaSource.h
new file mode 100644
--- /dev/null
+++ b/Engine/Plugins/Media/MediaPlayerEditor/Source/MediaPlayerEditor/Private/AssetDefinition/AssetDefinition_FileMediaSource.h
@@ -0,0 +1,22 @@
+// Copyright Epic Games, Inc. All Rights Reserved.
+
+#pragma once
+#include "FileMediaSource.h"
+#include "AssetDefinition/AssetDefinition_MediaSource.h"
+#include "AssetDefinition_FileMediaSource.generated.h"
+
+UCLASS()
+class UAssetDefinition_FileMediaSource : public UAssetDefinition_MediaSource
+{
+	GENERATED_BODY()
+public:
+	// UAssetDefinition Begin
+	virtual FText GetAssetDisplayName() const override { return NSLOCTEXT("AssetDefinition", "AssetDefinition_FileMediaSource", "File Media Source"); }
+	virtual TSoftClassPtr<UObject> GetAssetClass() const override { return UFileMediaSource::StaticClass(); }
+	virtual TConstArrayView<FAssetCategoryPath> GetAssetCategories() const override
+	{
+		static const auto Categories = { FAssetCategoryPath(EAssetCategoryPaths::Media, NSLOCTEXT("AssetDefinition", "AssetDefinition_FileMediaSourceSubMenu", "Media Source"), ECategoryMenuType::Section) };
+		return Categories;
+	}
+	// UAssetDefinition End
+};
diff --git a/Engine/Plugins/Media/MediaPlayerEditor/Source/MediaPlayerEditor/Private/AssetDefinition/AssetDefinition_MediaPlayer.cpp b/Engine/Plugins/Media/MediaPlayerEditor/Source/MediaPlayerEditor/Private/AssetDefinition/AssetDefinition_MediaPlayer.cpp
new file mode 100644
--- /dev/null
+++ b/Engine/Plugins/Media/MediaPlayerEditor/Source/MediaPlayerEditor/Private/AssetDefinition/AssetDefinition_MediaPlayer.cpp
@@ -0,0 +1,36 @@
+// Copyright Epic Games, Inc. All Rights Reserved.
+
+#include "AssetDefinition_MediaPlayer.h"
+#include "MediaPlayerEditorModule.h"
+#include "Modules/ModuleManager.h"
+#include "Toolkits/MediaPlayerEditorToolkit.h"
+
+#define LOCTEXT_NAMESPACE "AssetDefinition_MediaPlayer"
+
+EAssetCommandResult UAssetDefinition_MediaPlayer::OpenAssets(const FAssetOpenArgs& OpenArgs) const
+{
+	const EToolkitMode::Type Mode = OpenArgs.ToolkitHost.IsValid()
+		? EToolkitMode::WorldCentric
+		: EToolkitMode::Standalone;
+
+	const TArray<UObject*> Objects = OpenArgs.LoadObjects<UObject>();
+	for (auto ObjIt = Objects.CreateConstIterator(); ObjIt; ++ObjIt)
+	{
+		UMediaPlayer* MediaPlayer = Cast<UMediaPlayer>(*ObjIt);
+		IMediaPlayerEditorModule* MediaPlayerEditorModule = FModuleManager::LoadModulePtr<IMediaPlayerEditorModule>("MediaPlayerEditor");
+
+		if (MediaPlayer != nullptr && MediaPlayerEditorModule != nullptr)
+		{
+			TSharedPtr<ISlateStyle> MediaPlayerStyle = MediaPlayerEditorModule->GetStyle();
+			if (MediaPlayerStyle.IsValid())
+			{
+				const TSharedRef<FMediaPlayerEditorToolkit> EditorToolkit = MakeShareable(new FMediaPlayerEditorToolkit(MediaPlayerStyle.ToSharedRef()));
+				EditorToolkit->Initialize(MediaPlayer, Mode, OpenArgs.ToolkitHost);
+			}
+		}
+	}
+
+	return EAssetCommandResult::Handled;
+}
+
+#undef LOCTEXT_NAMESPACE
diff --git a/Engine/Plugins/Media/MediaPlayerEditor/Source/MediaPlayerEditor/Private/AssetDefinition/AssetDefinition_MediaPlayer.h b/Engine/Plugins/Media/MediaPlayerEditor/Source/MediaPlayerEditor/Private/AssetDefinition/AssetDefinition_MediaPlayer.h
new file mode 100644
--- /dev/null
+++ b/Engine/Plugins/Media/MediaPlayerEditor/Source/MediaPlayerEditor/Private/AssetDefinition/AssetDefinition_MediaPlayer.h
@@ -0,0 +1,24 @@
+// Copyright Epic Games, Inc. All Rights Reserved.
+
+#pragma once
+#include "AssetDefinitionDefault.h"
+#include "MediaPlayer.h"
+#include "AssetDefinition_MediaPlayer.generated.h"
+
+UCLASS()
+class UAssetDefinition_MediaPlayer : public UAssetDefinitionDefault
+{
+	GENERATED_BODY()
+public:
+	// UAssetDefinition Begin
+	virtual FText GetAssetDisplayName() const override { return NSLOCTEXT("AssetDefinition", "AssetDefinition_MediaPlayer", "Media Player"); }
+	virtual TSoftClassPtr<UObject> GetAssetClass() const override { return UMediaPlayer::StaticClass(); }
+	virtual FLinearColor GetAssetColor() const override { return FColor::Red; }
+	virtual TConstArrayView<FAssetCategoryPath> GetAssetCategories() const override
+	{
+		static const auto Categories = { FAssetCategoryPath(EAssetCategoryPaths::Media, NSLOCTEXT("AssetDefinition", "AssetDefinition_MediaPlayerSubMenu", "Other"), ECategoryMenuType::Section) };
+		return Categories;
+	}
+	virtual EAssetCommandResult OpenAssets(const FAssetOpenArgs& OpenArgs) const override;
+	// UAssetDefinition End
+};
diff --git a/Engine/Plugins/Media/MediaPlayerEditor/Source/MediaPlayerEditor/Private/AssetDefinition/AssetDefinition_MediaPlaylist.cpp b/Engine/Plugins/Media/MediaPlayerEditor/Source/MediaPlayerEditor/Private/AssetDefinition/AssetDefinition_MediaPlaylist.cpp
new file mode 100644
--- /dev/null
+++ b/Engine/Plugins/Media/MediaPlayerEditor/Source/MediaPlayerEditor/Private/AssetDefinition/AssetDefinition_MediaPlaylist.cpp
@@ -0,0 +1,36 @@
+// Copyright Epic Games, Inc. All Rights Reserved.
+
+#include "AssetDefinition_MediaPlaylist.h"
+#include "MediaPlayerEditorModule.h"
+#include "Modules/ModuleManager.h"
+#include "Toolkits/MediaPlaylistEditorToolkit.h"
+
+#define LOCTEXT_NAMESPACE "AssetDefinition_MediaPlaylist"
+
+EAssetCommandResult UAssetDefinition_MediaPlaylist::OpenAssets(const FAssetOpenArgs& OpenArgs) const
+{
+	const EToolkitMode::Type Mode = OpenArgs.ToolkitHost.IsValid()
+		? EToolkitMode::WorldCentric
+		: EToolkitMode::Standalone;
+
+	const TArray<UObject*> Objects = OpenArgs.LoadObjects<UObject>();
+	for (auto ObjIt = Objects.CreateConstIterator(); ObjIt; ++ObjIt)
+	{
+		UMediaPlaylist* MediaPlaylist = Cast<UMediaPlaylist>(*ObjIt);
+		IMediaPlayerEditorModule* MediaPlayerEditorModule = FModuleManager::LoadModulePtr<IMediaPlayerEditorModule>("MediaPlayerEditor");
+
+		if (MediaPlaylist != nullptr && MediaPlayerEditorModule != nullptr)
+		{
+			TSharedPtr<ISlateStyle> MediaPlayerStyle = MediaPlayerEditorModule->GetStyle();
+			if (MediaPlayerStyle.IsValid())
+			{
+				const TSharedRef<FMediaPlaylistEditorToolkit> EditorToolkit = MakeShareable(new FMediaPlaylistEditorToolkit(MediaPlayerStyle.ToSharedRef()));
+				EditorToolkit->Initialize(MediaPlaylist, Mode, OpenArgs.ToolkitHost);
+			}
+		}
+	}
+
+	return EAssetCommandResult::Handled;
+}
+
+#undef LOCTEXT_NAMESPACE
diff --git a/Engine/Plugins/Media/MediaPlayerEditor/Source/MediaPlayerEditor/Private/AssetDefinition/AssetDefinition_MediaPlaylist.h b/Engine/Plugins/Media/MediaPlayerEditor/Source/MediaPlayerEditor/Private/AssetDefinition/AssetDefinition_MediaPlaylist.h
new file mode 100644
--- /dev/null
+++ b/Engine/Plugins/Media/MediaPlayerEditor/Source/MediaPlayerEditor/Private/AssetDefinition/AssetDefinition_MediaPlaylist.h
@@ -0,0 +1,24 @@
+// Copyright Epic Games, Inc. All Rights Reserved.
+
+#pragma once
+#include "AssetDefinitionDefault.h"
+#include "MediaPlaylist.h"
+#include "AssetDefinition_MediaPlaylist.generated.h"
+
+UCLASS()
+class UAssetDefinition_MediaPlaylist : public UAssetDefinitionDefault
+{
+	GENERATED_BODY()
+public:
+	// UAssetDefinition Begin
+	virtual FText GetAssetDisplayName() const override { return NSLOCTEXT("AssetDefinition", "AssetDefinition_MediaPlaylist", "Media Playlist"); }
+	virtual TSoftClassPtr<UObject> GetAssetClass() const override { return UMediaPlaylist::StaticClass(); }
+	virtual FLinearColor GetAssetColor() const override { return FColor::Yellow; }
+	virtual TConstArrayView<FAssetCategoryPath> GetAssetCategories() const override
+	{
+		static const auto Categories = { FAssetCategoryPath(EAssetCategoryPaths::Media, NSLOCTEXT("AssetDefinition", "AssetDefinition_MediaPlaylistSubMenu", "Other"), ECategoryMenuType::Section) };
+		return Categories;
+	}
+	virtual EAssetCommandResult OpenAssets(const FAssetOpenArgs& OpenArgs) const override;
+	// UAssetDefinition End
+};
diff --git a/Engine/Plugins/Media/MediaPlayerEditor/Source/MediaPlayerEditor/Private/AssetDefinition/AssetDefinition_MediaSource.cpp b/Engine/Plugins/Media/MediaPlayerEditor/Source/MediaPlayerEditor/Private/AssetDefinition/AssetDefinition_MediaSource.cpp
new file mode 100644
--- /dev/null
+++ b/Engine/Plugins/Media/MediaPlayerEditor/Source/MediaPlayerEditor/Private/AssetDefinition/AssetDefinition_MediaSource.cpp
@@ -0,0 +1,39 @@
+// Copyright Epic Games, Inc. All Rights Reserved.
+
+#include "AssetDefinition/AssetDefinition_MediaSource.h"
+#include "MediaPlayerEditorModule.h"
+#include "Modules/ModuleManager.h"
+#include "Toolkits/MediaSourceEditorToolkit.h"
+
+#define LOCTEXT_NAMESPACE "AssetDefinition_MediaSource"
+
+EAssetCommandResult UAssetDefinition_MediaSource::OpenAssets(const FAssetOpenArgs& OpenArgs) const
+{
+	EToolkitMode::Type Mode = OpenArgs.ToolkitHost.IsValid()
+		? EToolkitMode::WorldCentric
+		: EToolkitMode::Standalone;
+
+	TArray<UObject*> Objects = OpenArgs.LoadObjects<UObject>();
+	for (auto ObjIt = Objects.CreateConstIterator(); ObjIt; ++ObjIt)
+	{
+		UMediaSource* MediaSource = Cast<UMediaSource>(*ObjIt);
+
+		if (MediaSource != nullptr)
+		{
+			IMediaPlayerEditorModule* MediaPlayerEditorModule = FModuleManager::LoadModulePtr<IMediaPlayerEditorModule>("MediaPlayerEditor");
+			if (MediaPlayerEditorModule != nullptr)
+			{
+				TSharedPtr<ISlateStyle> MediaSourceStyle = MediaPlayerEditorModule->GetStyle();
+				if (MediaSourceStyle.IsValid())
+				{
+					TSharedRef<FMediaSourceEditorToolkit> EditorToolkit = MakeShareable(new FMediaSourceEditorToolkit(MediaSourceStyle.ToSharedRef()));
+					EditorToolkit->Initialize(MediaSource, Mode, OpenArgs.ToolkitHost);
+				}
+			}
+		}
+	}
+
+	return EAssetCommandResult::Handled;
+}
+
+#undef LOCTEXT_NAMESPACE
diff --git a/Engine/Plugins/Media/MediaPlayerEditor/Source/MediaPlayerEditor/Private/AssetDefinition/AssetDefinition_MediaTexture.cpp b/Engine/Plugins/Media/MediaPlayerEditor/Source/MediaPlayerEditor/Private/AssetDefinition/AssetDefinition_MediaTexture.cpp
new file mode 100644
--- /dev/null
+++ b/Engine/Plugins/Media/MediaPlayerEditor/Source/MediaPlayerEditor/Private/AssetDefinition/AssetDefinition_MediaTexture.cpp
@@ -0,0 +1,34 @@
+// Copyright Epic Games, Inc. All Rights Reserved.
+
+#include "AssetDefinition_MediaTexture.h"
+
+#include "AnimationEditorUtils.h"
+#include "ContentBrowserMenuContexts.h"
+#include "ToolMenus.h"
+#include "Factories/MaterialFactoryNew.h"
+#include "Interfaces/ITextureEditorModule.h"
+#include "Materials/Material.h"
+#include "Modules/ModuleManager.h"
+
+#define LOCTEXT_NAMESPACE "AssetDefinition_MediaTexture"
+
+EAssetCommandResult UAssetDefinition_MediaTexture::OpenAssets(const FAssetOpenArgs& OpenArgs) const
+{
+	const EToolkitMode::Type Mode = OpenArgs.ToolkitHost.IsValid() ? EToolkitMode::WorldCentric : EToolkitMode::Standalone;
+
+	const TArray<UObject*> Objects = OpenArgs.LoadObjects<UObject>();
+	for (auto ObjIt = Objects.CreateConstIterator(); ObjIt; ++ObjIt)
+	{
+		UTexture* Texture = Cast<UTexture>(*ObjIt);
+
+		if (Texture != nullptr)
+		{
+			ITextureEditorModule* TextureEditorModule = &FModuleManager::LoadModuleChecked<ITextureEditorModule>("TextureEditor");
+			TextureEditorModule->CreateTextureEditor(Mode, OpenArgs.ToolkitHost, Texture);
+		}
+	}
+
+	return EAssetCommandResult::Handled;
+}
+
+#undef LOCTEXT_NAMESPACE
diff --git a/Engine/Plugins/Media/MediaPlayerEditor/Source/MediaPlayerEditor/Private/AssetDefinition/AssetDefinition_MediaTexture.h b/Engine/Plugins/Media/MediaPlayerEditor/Source/MediaPlayerEditor/Private/AssetDefinition/AssetDefinition_MediaTexture.h
new file mode 100644
--- /dev/null
+++ b/Engine/Plugins/Media/MediaPlayerEditor/Source/MediaPlayerEditor/Private/AssetDefinition/AssetDefinition_MediaTexture.h
@@ -0,0 +1,28 @@
+// Copyright Epic Games, Inc. All Rights Reserved.
+
+#pragma once
+#include "AssetDefinitionDefault.h"
+#include "MediaTexture.h"
+#include "AssetDefinition_MediaTexture.generated.h"
+
+UCLASS()
+class UAssetDefinition_MediaTexture : public UAssetDefinitionDefault
+{
+	GENERATED_BODY()
+public:
+	// UAssetDefinition Begin
+	virtual FText GetAssetDisplayName() const override { return NSLOCTEXT("AssetDefinition", "AssetDefinition_MediaTexture", "Media Texture"); }
+	virtual TSoftClassPtr<UObject> GetAssetClass() const override { return UMediaTexture::StaticClass(); }
+	virtual FLinearColor GetAssetColor() const override { return FColor::Red; }
+	virtual TConstArrayView<FAssetCategoryPath> GetAssetCategories() const override
+	{
+		static const auto Categories = 
+		{
+			FAssetCategoryPath(EAssetCategoryPaths::Media, NSLOCTEXT("AssetDefinition", "AssetDefinition_MediaTextureSubMenuMedia", "Other"), ECategoryMenuType::Section),
+			FAssetCategoryPath(EAssetCategoryPaths::Texture, NSLOCTEXT("AssetDefinition", "AssetDefinition_MediaTextureSubMenuTexture", "Other"), ECategoryMenuType::Section)
+		};
+		return Categories;
+	}
+	virtual EAssetCommandResult OpenAssets(const FAssetOpenArgs& OpenArgs) const override;
+	// UAssetDefinition End
+};
diff --git a/Engine/Plugins/Media/MediaPlayerEditor/Source/MediaPlayerEditor/Private/AssetDefinition/AssetDefinition_PlatformMediaSource.h b/Engine/Plugins/Media/MediaPlayerEditor/Source/MediaPlayerEditor/Private/AssetDefinition/AssetDefinition_PlatformMediaSource.h
new file mode 100644
--- /dev/null
+++ b/Engine/Plugins/Media/MediaPlayerEditor/Source/MediaPlayerEditor/Private/AssetDefinition/AssetDefinition_PlatformMediaSource.h
@@ -0,0 +1,22 @@
+// Copyright Epic Games, Inc. All Rights Reserved.
+
+#pragma once
+#include "PlatformMediaSource.h"
+#include "AssetDefinition/AssetDefinition_MediaSource.h"
+#include "AssetDefinition_PlatformMediaSource.generated.h"
+
+UCLASS()
+class UAssetDefinition_PlatformMediaSource : public UAssetDefinition_MediaSource
+{
+	GENERATED_BODY()
+public:
+	// UAssetDefinition Begin
+	virtual FText GetAssetDisplayName() const override { return NSLOCTEXT("AssetDefinition", "AssetDefinition_PlatformMediaSource", "Platform Media Source"); }
+	virtual TSoftClassPtr<UObject> GetAssetClass() const override { return UPlatformMediaSource::StaticClass(); }
+	virtual TConstArrayView<FAssetCategoryPath> GetAssetCategories() const override
+	{
+		static const auto Categories = { FAssetCategoryPath(EAssetCategoryPaths::Media, NSLOCTEXT("AssetDefinition", "AssetDefinition_PlatformMediaSourceSubMenu", "Media Source"), ECategoryMenuType::Section) };
+		return Categories;
+	}
+	// UAssetDefinition End
+};
diff --git a/Engine/Plugins/Media/MediaPlayerEditor/Source/MediaPlayerEditor/Private/AssetDefinition/AssetDefinition_StreamMediaSource.h b/Engine/Plugins/Media/MediaPlayerEditor/Source/MediaPlayerEditor/Private/AssetDefinition/AssetDefinition_StreamMediaSource.h
new file mode 100644
--- /dev/null
+++ b/Engine/Plugins/Media/MediaPlayerEditor/Source/MediaPlayerEditor/Private/AssetDefinition/AssetDefinition_StreamMediaSource.h
@@ -0,0 +1,22 @@
+// Copyright Epic Games, Inc. All Rights Reserved.
+
+#pragma once
+#include "StreamMediaSource.h"
+#include "AssetDefinition/AssetDefinition_MediaSource.h"
+#include "AssetDefinition_StreamMediaSource.generated.h"
+
+UCLASS()
+class UAssetDefinition_StreamMediaSource : public UAssetDefinition_MediaSource
+{
+	GENERATED_BODY()
+public:
+	// UAssetDefinition Begin
+	virtual FText GetAssetDisplayName() const override { return NSLOCTEXT("AssetDefinition", "AssetDefinition_StreamMediaSource", "Stream Media Source"); }
+	virtual TSoftClassPtr<UObject> GetAssetClass() const override { return UStreamMediaSource::StaticClass(); }
+	virtual TConstArrayView<FAssetCategoryPath> GetAssetCategories() const override
+	{
+		static const auto Categories = { FAssetCategoryPath(EAssetCategoryPaths::Media, NSLOCTEXT("AssetDefinition", "AssetDefinition_StreamMediaSourceSubMenu", "Media Source"), ECategoryMenuType::Section) };
+		return Categories;
+	}
+	// UAssetDefinition End
+};
diff --git a/Engine/Plugins/Media/MediaPlayerEditor/Source/MediaPlayerEditor/Private/AssetTools/FileMediaSourceActions.cpp b/Engine/Plugins/Media/MediaPlayerEditor/Source/MediaPlayerEditor/Private/AssetTools/FileMediaSourceActions.cpp
deleted file mode 100644
diff --git a/Engine/Plugins/Media/MediaPlayerEditor/Source/MediaPlayerEditor/Private/AssetTools/FileMediaSourceActions.h b/Engine/Plugins/Media/MediaPlayerEditor/Source/MediaPlayerEditor/Private/AssetTools/FileMediaSourceActions.h
deleted file mode 100644
diff --git a/Engine/Plugins/Media/MediaPlayerEditor/Source/MediaPlayerEditor/Private/AssetTools/MediaPlayerActions.cpp b/Engine/Plugins/Media/MediaPlayerEditor/Source/MediaPlayerEditor/Private/AssetTools/MediaPlayerActions.cpp
deleted file mode 100644
diff --git a/Engine/Plugins/Media/MediaPlayerEditor/Source/MediaPlayerEditor/Private/AssetTools/MediaPlayerActions.h b/Engine/Plugins/Media/MediaPlayerEditor/Source/MediaPlayerEditor/Private/AssetTools/MediaPlayerActions.h
deleted file mode 100644
diff --git a/Engine/Plugins/Media/MediaPlayerEditor/Source/MediaPlayerEditor/Private/AssetTools/MediaPlaylistActions.cpp b/Engine/Plugins/Media/MediaPlayerEditor/Source/MediaPlayerEditor/Private/AssetTools/MediaPlaylistActions.cpp
deleted file mode 100644
diff --git a/Engine/Plugins/Media/MediaPlayerEditor/Source/MediaPlayerEditor/Private/AssetTools/MediaPlaylistActions.h b/Engine/Plugins/Media/MediaPlayerEditor/Source/MediaPlayerEditor/Private/AssetTools/MediaPlaylistActions.h
deleted file mode 100644
diff --git a/Engine/Plugins/Media/MediaPlayerEditor/Source/MediaPlayerEditor/Private/AssetTools/MediaTextureActions.cpp b/Engine/Plugins/Media/MediaPlayerEditor/Source/MediaPlayerEditor/Private/AssetTools/MediaTextureActions.cpp
deleted file mode 100644
diff --git a/Engine/Plugins/Media/MediaPlayerEditor/Source/MediaPlayerEditor/Private/AssetTools/MediaTextureActions.h b/Engine/Plugins/Media/MediaPlayerEditor/Source/MediaPlayerEditor/Private/AssetTools/MediaTextureActions.h
deleted file mode 100644
diff --git a/Engine/Plugins/Media/MediaPlayerEditor/Source/MediaPlayerEditor/Private/AssetTools/PlatformMediaSourceActions.cpp b/Engine/Plugins/Media/MediaPlayerEditor/Source/MediaPlayerEditor/Private/AssetTools/PlatformMediaSourceActions.cpp
deleted file mode 100644
diff --git a/Engine/Plugins/Media/MediaPlayerEditor/Source/MediaPlayerEditor/Private/AssetTools/PlatformMediaSourceActions.h b/Engine/Plugins/Media/MediaPlayerEditor/Source/MediaPlayerEditor/Private/AssetTools/PlatformMediaSourceActions.h
deleted file mode 100644
diff --git a/Engine/Plugins/Media/MediaPlayerEditor/Source/MediaPlayerEditor/Private/AssetTools/StreamMediaSourceActions.cpp b/Engine/Plugins/Media/MediaPlayerEditor/Source/MediaPlayerEditor/Private/AssetTools/StreamMediaSourceActions.cpp
deleted file mode 100644
diff --git a/Engine/Plugins/Media/MediaPlayerEditor/Source/MediaPlayerEditor/Private/AssetTools/StreamMediaSourceActions.h b/Engine/Plugins/Media/MediaPlayerEditor/Source/MediaPlayerEditor/Private/AssetTools/StreamMediaSourceActions.h
deleted file mode 100644
diff --git a/Engine/Plugins/Media/MediaPlayerEditor/Source/MediaPlayerEditor/Private/MediaPlayerEditorModule.cpp b/Engine/Plugins/Media/MediaPlayerEditor/Source/MediaPlayerEditor/Private/MediaPlayerEditorModule.cpp
--- a/Engine/Plugins/Media/MediaPlayerEditor/Source/MediaPlayerEditor/Private/MediaPlayerEditorModule.cpp
+++ b/Engine/Plugins/Media/MediaPlayerEditor/Source/MediaPlayerEditor/Private/MediaPlayerEditorModule.cpp
@@ -27,13 +27,7 @@
 #include "MediaTexture.h"
 #include "PlatformMediaSource.h"
 
-#include "AssetTools/FileMediaSourceActions.h"
-#include "AssetTools/MediaPlayerActions.h"
-#include "AssetTools/MediaPlaylistActions.h"
 #include "AssetTools/MediaSourceActions.h"
-#include "AssetTools/MediaTextureActions.h"
-#include "AssetTools/PlatformMediaSourceActions.h"
-#include "AssetTools/StreamMediaSourceActions.h"
 
 #include "Customizations/BaseMediaSourceCustomization.h"
 #include "Customizations/FileMediaSourceCustomization.h"
@@ -145,13 +139,9 @@ protected:
 	{
 		IAssetTools& AssetTools = FModuleManager::LoadModuleChecked<FAssetToolsModule>("AssetTools").Get();
 
-		RegisterAssetTypeAction(AssetTools, MakeShareable(new FFileMediaSourceActions()));
-		RegisterAssetTypeAction(AssetTools, MakeShareable(new FMediaPlayerActions(Style.ToSharedRef())));
-		RegisterAssetTypeAction(AssetTools, MakeShareable(new FMediaPlaylistActions(Style.ToSharedRef())));
+		PRAGMA_DISABLE_DEPRECATION_WARNINGS
 		RegisterAssetTypeAction(AssetTools, MakeShareable(new FMediaSourceActions()));
-		RegisterAssetTypeAction(AssetTools, MakeShareable(new FMediaTextureActions));
-		RegisterAssetTypeAction(AssetTools, MakeShareable(new FPlatformMediaSourceActions()));
-		RegisterAssetTypeAction(AssetTools, MakeShareable(new FStreamMediaSourceActions()));
+		PRAGMA_ENABLE_DEPRECATION_WARNINGS
 	}
 
 	/**
diff --git a/Engine/Plugins/Media/MediaPlayerEditor/Source/MediaPlayerEditor/Public/AssetDefinition/AssetDefinition_MediaSource.h b/Engine/Plugins/Media/MediaPlayerEditor/Source/MediaPlayerEditor/Public/AssetDefinition/AssetDefinition_MediaSource.h
new file mode 100644
--- /dev/null
+++ b/Engine/Plugins/Media/MediaPlayerEditor/Source/MediaPlayerEditor/Public/AssetDefinition/AssetDefinition_MediaSource.h
@@ -0,0 +1,20 @@
+// Copyright Epic Games, Inc. All Rights Reserved.
+
+#pragma once
+#include "AssetDefinitionDefault.h"
+#include "MediaSource.h"
+#include "AssetDefinition_MediaSource.generated.h"
+
+UCLASS(MinimalAPI)
+class UAssetDefinition_MediaSource : public UAssetDefinitionDefault
+{
+	GENERATED_BODY()
+
+public:
+	// UAssetDefinition Begin
+	virtual FText GetAssetDisplayName() const override { return FText::GetEmpty(); }
+	virtual FLinearColor GetAssetColor() const override { return FColor::White; }
+	virtual TSoftClassPtr<UObject> GetAssetClass() const override { return UMediaSource::StaticClass(); }
+	MEDIAPLAYEREDITOR_API virtual EAssetCommandResult OpenAssets(const FAssetOpenArgs& OpenArgs) const override;
+	// UAssetDefinition End
+};
